Initial integration of Control Rig in Anim Next
- Created AnimNextControlRig plugin
- Added a new AlphaBlendTrait that enables alpha blend options to base traits that generate an keyframe from the current keyframe at the top of the VM stack
- Small refactors of ControlRig helpers to enable AnimNext work
- Created a ControlRigTrait and ControlRigTask to enable the use of Control Rig in AnimNext
- Created a ControlRigTrait details customization to allow the custom data required to expose pins from Control Rig into AnimNext
- Various fixes in Anim Next to fix compilation problems with latent pins

diff --git a/Engine/Plugins/Animation/ControlRig/Source/ControlRig/Private/AnimNode_ControlRig.cpp b/Engine/Plugins/Animation/ControlRig/Source/ControlRig/Private/AnimNode_ControlRig.cpp
--- a/Engine/Plugins/Animation/ControlRig/Source/ControlRig/Private/AnimNode_ControlRig.cpp
+++ b/Engine/Plugins/Animation/ControlRig/Source/ControlRig/Private/AnimNode_ControlRig.cpp
@@ -14,7 +14,6 @@
 #include "Animation/AnimCurveUtils.h"
 #include "Animation/AnimStats.h"
 #include "Units/Execution/RigUnit_PrepareForExecution.h"
-#include "Tools/ControlRigVariableMappings.h"
 
 #include UE_INLINE_GENERATED_CPP_BY_NAME(AnimNode_ControlRig)
 
@@ -470,6 +469,7 @@ void FAnimNode_ControlRig::PropagateInputProperties(const UObject* InSourceInsta
 	DECLARE_SCOPE_HIERARCHICAL_COUNTER_FUNC()
 	QUICK_SCOPE_CYCLE_COUNTER(STAT_AnimNode_ControlRig_PropagateInputProperties);
 
+
 	if (!InSourceInstance)
 	{
 		return;
diff --git a/Engine/Plugins/Animation/ControlRig/Source/ControlRig/Private/Tools/ControlRigHierarchyMappings.cpp b/Engine/Plugins/Animation/ControlRig/Source/ControlRig/Private/Tools/ControlRigHierarchyMappings.cpp
--- a/Engine/Plugins/Animation/ControlRig/Source/ControlRig/Private/Tools/ControlRigHierarchyMappings.cpp
+++ b/Engine/Plugins/Animation/ControlRig/Source/ControlRig/Private/Tools/ControlRigHierarchyMappings.cpp
@@ -92,8 +92,8 @@ void FControlRigHierarchyMappings::UpdateControlRigRefPoseIfNeeded(UControlRig*
 void FControlRigHierarchyMappings::UpdateInputOutputMappingIfRequired(UControlRig* InControlRig
 	, URigHierarchy* InHierarchy
 	, const FBoneContainer& InRequiredBones
-	, TArray<FBoneReference>& InInputBonesToTransfer
-	, TArray<FBoneReference>& InOutputBonesToTransfer
+	, const TArray<FBoneReference>& InInputBonesToTransfer
+	, const TArray<FBoneReference>& InOutputBonesToTransfer
 	, const TWeakObjectPtr<UNodeMappingContainer>& InNodeMappingContainer
 	, bool bInTransferPoseInGlobalSpace
 	, bool bInResetInputPoseToInitial)
@@ -535,8 +535,8 @@ void FControlRigHierarchyMappings::UpdateOutput(UControlRig* ControlRig
 void FControlRigHierarchyMappings::UpdateInputOutputMappingIfRequiredImpl(UControlRig* InControlRig
 	, URigHierarchy* InHierarchy
 	, const FBoneContainer& InRequiredBones
-	, TArray<FBoneReference>& InInputBonesToTransfer
-	, TArray<FBoneReference>& InOutputBonesToTransfer
+	, const TArray<FBoneReference>& InInputBonesToTransfer
+	, const TArray<FBoneReference>& InOutputBonesToTransfer
 	, const TWeakObjectPtr<UNodeMappingContainer>& InNodeMappingContainer)
 {
 	DECLARE_SCOPE_HIERARCHICAL_COUNTER_FUNC()
@@ -592,7 +592,7 @@ void FControlRigHierarchyMappings::UpdateInputOutputMappingIfRequiredImpl(UContr
 		}
 
 		auto UpdatingMappingFromSpecificTransferList = [](
-			TArray<FBoneReference>& InTransferList,
+			const TArray<FBoneReference>& InTransferList,
 			const TWeakObjectPtr<UNodeMappingContainer>& InMappingContainer,
 			const FBoneContainer& InRequiredBones,
 			const FReferenceSkeleton& InRefSkeleton,
@@ -608,16 +608,17 @@ void FControlRigHierarchyMappings::UpdateInputOutputMappingIfRequiredImpl(UContr
 					TMap<FName, FName> TargetToSourceMappingTable;
 					InMappingContainer->GetTargetToSourceMappingTable(TargetToSourceMappingTable);
 
-					for (FBoneReference& InputBoneToTransfer : InTransferList)
+					for (const FBoneReference& InputBoneToTransfer : InTransferList)
 					{
-						if (!InputBoneToTransfer.Initialize(InRequiredBones))
+						const int32 BoneIndex = InRequiredBones.GetPoseBoneIndexForBoneName(InputBoneToTransfer.BoneName);
+						if (BoneIndex == INDEX_NONE)
 						{
 							continue;
 						}
-						const FName TargetNodeName = InRefSkeleton.GetBoneName(InputBoneToTransfer.BoneIndex);
+						const FName TargetNodeName = InRefSkeleton.GetBoneName(BoneIndex/*InputBoneToTransfer.BoneIndex*/);
 						if (const FName* SourceName = TargetToSourceMappingTable.Find(TargetNodeName))
 						{
-							OutMapping.Add(*SourceName, InputBoneToTransfer.BoneIndex);
+							OutMapping.Add(*SourceName, BoneIndex/*InputBoneToTransfer.BoneIndex*/);
 						}
 					}
 				}
@@ -627,18 +628,19 @@ void FControlRigHierarchyMappings::UpdateInputOutputMappingIfRequiredImpl(UContr
 					TArray<FNodeItem> NodeItems;
 					InControlRig->GetMappableNodeData(NodeNames, NodeItems);
 
-					for (FBoneReference& InputBoneToTransfer : InTransferList)
+					for (const FBoneReference& InputBoneToTransfer : InTransferList)
 					{
-						if (!InputBoneToTransfer.Initialize(InRequiredBones))
+						const int32 BoneIndex = InRequiredBones.GetPoseBoneIndexForBoneName(InputBoneToTransfer.BoneName);
+						if (BoneIndex == INDEX_NONE)
 						{
 							continue;
 						}
-						if (InRequiredBonesArray.IsValidIndex(InputBoneToTransfer.BoneIndex))
+						if (InRequiredBonesArray.IsValidIndex(BoneIndex/*InputBoneToTransfer.BoneIndex*/))
 						{
-							const FName& BoneName = InRefSkeleton.GetBoneName(InRequiredBonesArray[InputBoneToTransfer.BoneIndex]);
+							const FName& BoneName = InRefSkeleton.GetBoneName(InRequiredBonesArray[BoneIndex/*InputBoneToTransfer.BoneIndex*/]);
 							if (NodeNames.Contains(BoneName))
 							{
-								OutMapping.Add(BoneName, InputBoneToTransfer.BoneIndex);
+								OutMapping.Add(BoneName, BoneIndex/*InputBoneToTransfer.BoneIndex*/);
 							}
 						}
 					}
diff --git a/Engine/Plugins/Animation/ControlRig/Source/ControlRig/Private/Tools/ControlRigVariableMappings.cpp b/Engine/Plugins/Animation/ControlRig/Source/ControlRig/Private/Tools/ControlRigVariableMappings.cpp
--- a/Engine/Plugins/Animation/ControlRig/Source/ControlRig/Private/Tools/ControlRigVariableMappings.cpp
+++ b/Engine/Plugins/Animation/ControlRig/Source/ControlRig/Private/Tools/ControlRigVariableMappings.cpp
@@ -26,6 +26,13 @@ void SetControlValue(FRigControlElement* InControlElement, FProperty* InSourcePr
 	InTargetHierarchy->SetControlValue(InControlElement, Value, ERigControlValueType::Current);
 }
 
+template<typename ValueType, typename HierarchyValueType = ValueType>
+void SetControlValue(FRigControlElement* InControlElement, const uint8* ValuePtr, URigHierarchy* InTargetHierarchy)
+{
+	const FRigControlValue Value = FRigControlValue::Make<HierarchyValueType>(*((const ValueType*)ValuePtr));
+	InTargetHierarchy->SetControlValue(InControlElement, Value, ERigControlValueType::Current);
+}
+
 template<typename ValueType, typename VariableValueType = ValueType>
 void SetVariableValue(FRigVMExternalVariable& InVariable, FProperty* InSourceProperty, const UObject* InSourceInstance)
 {
@@ -33,6 +40,13 @@ void SetVariableValue(FRigVMExternalVariable& InVariable, FProperty* InSourcePro
 	InVariable.SetValue<VariableValueType>(*ValuePtr);
 }
 
+template<typename ValueType, typename VariableValueType = ValueType>
+void SetVariableValue(FRigVMExternalVariable& InVariable, const uint8* InValuePtr)
+{
+	const ValueType* ValuePtr = (const ValueType*)InValuePtr;
+	InVariable.SetValue<VariableValueType>(*ValuePtr);
+}
+
 } // end namespace AnimNodeLocals
 
 
@@ -66,7 +80,7 @@ void FControlRigVariableMappings::InitializeProperties(const UClass* InSourceCla
 
 		if (bUseFunctions)
 		{
-			if (TargetHierarchy && SourceProperty)
+			if (TargetControlRig && TargetHierarchy && SourceProperty)
 			{
 				const FName& DestPropertyName = InDestPropertyNames[Idx];
 				if (FRigControlElement* ControlElement = TargetControlRig->FindControl(DestPropertyName))
@@ -86,6 +100,52 @@ void FControlRigVariableMappings::InitializeProperties(const UClass* InSourceCla
 	}
 }
 
+void FControlRigVariableMappings::InitializeCustomProperties(UControlRig* TargetControlRig, const FCustomPropertyMappings& CustomPropertyMapping)
+{
+	const TArray<FCustomPropertyData>& Mappings = CustomPropertyMapping.GetMappings();
+
+	ResetCustomProperties(Mappings.Num());
+
+	if (TargetControlRig != nullptr)
+	{
+		URigHierarchy* TargetHierarchy = TargetControlRig->GetHierarchy();
+
+		for (const FCustomPropertyData& Mapping : Mappings)
+		{
+			switch (Mapping.Type)
+			{
+				case FCustomPropertyData::ECustomPropertyType::Variable:
+				{
+					const FRigVMExternalVariable Variable = TargetControlRig->GetPublicVariableByName(Mapping.TargetName);
+					if (Variable.IsValid() && !Variable.bIsReadOnly)
+					{
+						AddCustomVariableFunction(Variable, Mapping.Property, Mapping.SourceMemory);
+					}
+					break;
+				}
+				case FCustomPropertyData::ECustomPropertyType::Control:
+				{
+					if (FRigControlElement* ControlElement = TargetControlRig->FindControl(Mapping.TargetName))
+					{
+						AddCustomControlFunction(ControlElement, Mapping.ControlType, Mapping.SourceMemory, TargetHierarchy);
+					}
+					break;
+				}
+				default:
+				{
+					ensureMsgf(false, TEXT("Unsupported or Invalid Custom Property Type"));
+					break;
+				}
+			}
+		}
+	}
+}
+
+void FControlRigVariableMappings::ResetCustomProperties(int32 NewSize)
+{
+	CustomUpdateFunctions.Reset(NewSize);
+}
+
 void FControlRigVariableMappings::PropagateInputProperties(const UObject* InSourceInstance, UControlRig* InTargetControlRig, TArray<FName>& InDestPropertyNames)
 {
 	DECLARE_SCOPE_HIERARCHICAL_COUNTER_FUNC()
@@ -120,6 +180,28 @@ void FControlRigVariableMappings::PropagateInputProperties(const UObject* InSour
 	}
 }
 
+void FControlRigVariableMappings::PropagateCustomInputProperties(UControlRig* InTargetControlRig)
+{
+	if (!InTargetControlRig)
+	{
+		return;
+	}
+	
+	URigHierarchy* TargetHierarchy = InTargetControlRig->GetHierarchy();
+	if (!TargetHierarchy)
+	{
+		return;
+	}
+
+	if (!CustomUpdateFunctions.IsEmpty())
+	{
+		Algo::ForEach(CustomUpdateFunctions, [](const CustomPropertyUpdateFunction& InFunc)
+			{
+				InFunc();
+			});
+	}
+}
+
 void FControlRigVariableMappings::UpdateCurveInputs(UControlRig* InControlRig, const TMap<FName, FName> &InputMapping, const FBlendedCurve& InCurveData)
 {
 	DECLARE_SCOPE_HIERARCHICAL_COUNTER_FUNC()
@@ -544,6 +626,11 @@ void FControlRigVariableMappings::AddUpdateFunction(PropertyUpdateFunction&& InF
 	UpdateFunctions.Add(MoveTemp(InFunction));
 }
 
+void FControlRigVariableMappings::AddCustomUpdateFunction(CustomPropertyUpdateFunction&& InFunction)
+{
+	CustomUpdateFunctions.Add(MoveTemp(InFunction));
+}
+
 void FControlRigVariableMappings::AddVariableFunction(const FRigVMExternalVariable& InVariable, FProperty* InSourceProperty)
 {
 	using namespace AnimNodeLocals;
@@ -659,6 +746,209 @@ void FControlRigVariableMappings::AddVariableFunction(const FRigVMExternalVariab
 	}
 }
 
+void FControlRigVariableMappings::AddCustomControlFunction(FRigControlElement* InControlElement, ERigControlType InControlType, const uint8* InSourcePropertyMemory, URigHierarchy* InTargetHierarchy)
+{
+	using namespace AnimNodeLocals;
+
+	if (!ensure(InControlType == InControlElement->Settings.ControlType))
+	{
+		return;
+	}
+
+	switch (InControlElement->Settings.ControlType)
+	{
+		case ERigControlType::Bool:
+		{
+			AddCustomUpdateFunction([InSourcePropertyMemory, InTargetHierarchy, InControlElement]()
+				{
+					SetControlValue<bool>(InControlElement, InSourcePropertyMemory, InTargetHierarchy);
+				});
+			break;
+		}
+		case ERigControlType::Float:
+		case ERigControlType::ScaleFloat:
+		{
+			AddCustomUpdateFunction([InSourcePropertyMemory, InTargetHierarchy, InControlElement]()
+				{
+					SetControlValue<float>(InControlElement, InSourcePropertyMemory, InTargetHierarchy);
+				});
+			break;
+		}
+		case ERigControlType::Integer:
+		{
+			AddCustomUpdateFunction([InSourcePropertyMemory, InTargetHierarchy, InControlElement]()
+				{
+					SetControlValue<int32>(InControlElement, InSourcePropertyMemory, InTargetHierarchy);
+				});
+			break;
+		}
+		case ERigControlType::Vector2D:
+		{
+			AddCustomUpdateFunction([InSourcePropertyMemory, InTargetHierarchy, InControlElement]()
+				{
+					SetControlValue<FVector2D>(InControlElement, InSourcePropertyMemory, InTargetHierarchy);
+				});
+			break;
+		}
+		case ERigControlType::Position:
+		case ERigControlType::Scale:
+		{
+			return AddCustomUpdateFunction([InSourcePropertyMemory, InTargetHierarchy, InControlElement]()
+				{
+					SetControlValue<FVector>(InControlElement, InSourcePropertyMemory, InTargetHierarchy);
+				});
+			break;
+		}
+		case ERigControlType::Rotator:
+		{
+			return AddCustomUpdateFunction([InSourcePropertyMemory, InTargetHierarchy, InControlElement]()
+				{
+					SetControlValue<FRotator>(InControlElement, InSourcePropertyMemory, InTargetHierarchy);
+				});
+			break;
+		}
+		case ERigControlType::Transform:
+		{
+			return AddCustomUpdateFunction([InSourcePropertyMemory, InTargetHierarchy, InControlElement]()
+				{
+					SetControlValue<FTransform>(InControlElement, InSourcePropertyMemory, InTargetHierarchy);
+				});
+			break;
+		}
+		case ERigControlType::TransformNoScale:
+		{
+			return AddCustomUpdateFunction([InSourcePropertyMemory, InTargetHierarchy, InControlElement]()
+				{
+					SetControlValue<FTransform, FTransformNoScale>(InControlElement, InSourcePropertyMemory, InTargetHierarchy);
+				});
+			break;
+		}
+		case ERigControlType::EulerTransform:
+		{
+			return AddCustomUpdateFunction([InSourcePropertyMemory, InTargetHierarchy, InControlElement]()
+				{
+					const FTransform* ValuePtr = (const FTransform*)InSourcePropertyMemory;
+					const FRigControlValue Value = FRigControlValue::Make<FEulerTransform>(FEulerTransform(*ValuePtr));
+					InTargetHierarchy->SetControlValue(InControlElement, Value, ERigControlValueType::Current);
+				});
+			break;
+		}
+		default:
+		{
+			checkNoEntry();
+		}
+	}
+}
+
+void FControlRigVariableMappings::AddCustomVariableFunction(const FRigVMExternalVariable& InVariable, const FProperty* InSourceProperty, const uint8* InSourcePropertyMemory)
+{
+	using namespace AnimNodeLocals;
+
+	if (CastField<FBoolProperty>(InSourceProperty) && InVariable.TypeName == RigVMTypeUtils::BoolTypeName)
+	{
+		AddCustomUpdateFunction([InSourcePropertyMemory, Variable = InVariable]() mutable
+			{
+				SetVariableValue<bool>(Variable, InSourcePropertyMemory);
+			});
+	}
+	else if (CastField<FFloatProperty>(InSourceProperty) && (InVariable.TypeName == RigVMTypeUtils::FloatTypeName || InVariable.TypeName == RigVMTypeUtils::DoubleTypeName))
+	{
+		if (InVariable.TypeName == RigVMTypeUtils::FloatTypeName)
+		{
+			AddCustomUpdateFunction([InSourcePropertyMemory, Variable = InVariable]() mutable
+				{
+					SetVariableValue<float>(Variable, InSourcePropertyMemory);
+				});
+		}
+		else
+		{
+			AddCustomUpdateFunction([InSourcePropertyMemory, Variable = InVariable]() mutable
+				{
+					SetVariableValue<float, double>(Variable, InSourcePropertyMemory);
+				});
+		}
+	}
+	else if (CastField<FDoubleProperty>(InSourceProperty) && (InVariable.TypeName == RigVMTypeUtils::FloatTypeName || InVariable.TypeName == RigVMTypeUtils::DoubleTypeName))
+	{
+		if (InVariable.TypeName == RigVMTypeUtils::FloatTypeName)
+		{
+			AddCustomUpdateFunction([InSourcePropertyMemory, Variable = InVariable]() mutable
+				{
+					SetVariableValue<double, float>(Variable, InSourcePropertyMemory);
+				});
+		}
+		else
+		{
+			AddCustomUpdateFunction([InSourcePropertyMemory, Variable = InVariable]() mutable
+				{
+					SetVariableValue<double>(Variable, InSourcePropertyMemory);
+				});
+		}
+	}
+	else if (CastField<FIntProperty>(InSourceProperty) && InVariable.TypeName == RigVMTypeUtils::Int32TypeName)
+	{
+		AddCustomUpdateFunction([InSourcePropertyMemory, Variable = InVariable]() mutable
+			{
+				SetVariableValue<int32>(Variable, InSourcePropertyMemory);
+			});
+	}
+	else if (CastField<FNameProperty>(InSourceProperty) && InVariable.TypeName == RigVMTypeUtils::FNameTypeName)
+	{
+		AddCustomUpdateFunction([InSourcePropertyMemory, Variable = InVariable]() mutable
+			{
+				SetVariableValue<FName>(Variable, InSourcePropertyMemory);
+			});
+	}
+	else if (CastField<FNameProperty>(InSourceProperty) && InVariable.TypeName == RigVMTypeUtils::FStringTypeName)
+	{
+		AddCustomUpdateFunction([InSourcePropertyMemory, Variable = InVariable]() mutable
+			{
+				SetVariableValue<FString>(Variable, InSourcePropertyMemory);
+			});
+	}
+	else if (const FStructProperty* StructProperty = CastField<FStructProperty>(InSourceProperty))
+	{
+		AddCustomUpdateFunction([StructProperty, InSourcePropertyMemory, Variable = InVariable]()
+			{
+				StructProperty->Struct->CopyScriptStruct(Variable.Memory, InSourcePropertyMemory, 1);
+			});
+	}
+	else if (const FArrayProperty* ArrayProperty = CastField<FArrayProperty>(InSourceProperty))
+	{
+		if (ensure(ArrayProperty->SameType(InVariable.Property)))
+		{
+			AddCustomUpdateFunction([ArrayProperty, InSourcePropertyMemory, Variable = InVariable]()
+				{
+					ArrayProperty->CopyCompleteValue(Variable.Memory, InSourcePropertyMemory);
+				});
+		}
+	}
+	else if (const FObjectProperty* ObjectProperty = CastField<FObjectProperty>(InSourceProperty))
+	{
+		if (ensure(ObjectProperty->SameType(InVariable.Property)))
+		{
+			AddCustomUpdateFunction([ObjectProperty, InSourcePropertyMemory, Variable = InVariable]()
+				{
+					ObjectProperty->CopyCompleteValue(Variable.Memory, InSourcePropertyMemory);
+				});
+		}
+	}
+	else if (const FEnumProperty* EnumProperty = CastField<FEnumProperty>(InSourceProperty))
+	{
+		if (ensure(EnumProperty->SameType(InVariable.Property)))
+		{
+			AddCustomUpdateFunction([EnumProperty, InSourcePropertyMemory, Variable = InVariable]()
+				{
+					EnumProperty->CopyCompleteValue(Variable.Memory, InSourcePropertyMemory);
+				});
+		}
+	}
+	else
+	{
+		ensureMsgf(false, TEXT("Property %s type %s not recognized"), *InSourceProperty->GetName(), *InSourceProperty->GetCPPType());
+	}
+}
+
 void FControlRigVariableMappings::CacheCurveMappings(const TMap<FName, FName>& InInputMapping, const TMap<FName, FName>& InOutputMapping, URigHierarchy* InHierarchy)
 {
 	CacheCurveMappings(InInputMapping, InputCurveMappings, InHierarchy);
diff --git a/Engine/Plugins/Animation/ControlRig/Source/ControlRig/Public/ControlRig.h b/Engine/Plugins/Animation/ControlRig/Source/ControlRig/Public/ControlRig.h
--- a/Engine/Plugins/Animation/ControlRig/Source/ControlRig/Public/ControlRig.h
+++ b/Engine/Plugins/Animation/ControlRig/Source/ControlRig/Public/ControlRig.h
@@ -603,6 +603,12 @@ private:
 
 	void CopyPoseFromOtherRig(UControlRig* Subject);
 
+public:
+	TGuardValue<bool> GetResetCurrentTransformsAfterConstructionGuard(bool bNewValue)
+	{
+		return TGuardValue<bool> (bResetCurrentTransformsAfterConstruction, bNewValue);
+	}
+
 protected:
 	bool bCopyHierarchyBeforeConstruction;
 	bool bResetInitialTransformsBeforeConstruction;
diff --git a/Engine/Plugins/Animation/ControlRig/Source/ControlRig/Public/Tools/ControlRigHierarchyMappings.h b/Engine/Plugins/Animation/ControlRig/Source/ControlRig/Public/Tools/ControlRigHierarchyMappings.h
--- a/Engine/Plugins/Animation/ControlRig/Source/ControlRig/Public/Tools/ControlRigHierarchyMappings.h
+++ b/Engine/Plugins/Animation/ControlRig/Source/ControlRig/Public/Tools/ControlRigHierarchyMappings.h
@@ -20,7 +20,7 @@ class TAutoConsoleVariable;
 
 extern TAutoConsoleVariable<int32> CVarControlRigEnableAnimNodePerformanceOptimizations;
 
-struct FControlRigHierarchyMappings
+struct CONTROLRIG_API FControlRigHierarchyMappings
 {
 	FControlRigHierarchyMappings() = default;
 
@@ -53,8 +53,8 @@ struct FControlRigHierarchyMappings
 	void UpdateInputOutputMappingIfRequired(UControlRig* InControlRig
 		, URigHierarchy* InHierarchy
 		, const FBoneContainer& InRequiredBones
-		, TArray<FBoneReference>& InInputBonesToTransfer
-		, TArray<FBoneReference>& InOutputBonesToTransfer
+		, const TArray<FBoneReference>& InInputBonesToTransfer
+		, const TArray<FBoneReference>& InOutputBonesToTransfer
 		, const TWeakObjectPtr<UNodeMappingContainer>& InNodeMappingContainer
 		, bool bInTransferPoseInGlobalSpace
 		, bool bResetInputPoseToInitial);
@@ -113,8 +113,8 @@ private:
 	void UpdateInputOutputMappingIfRequiredImpl(UControlRig* InControlRig
 		, URigHierarchy* InHierarchy
 		, const FBoneContainer& InRequiredBones
-		, TArray<FBoneReference>& InInputBonesToTransfer
-		, TArray<FBoneReference>& InOutputBonesToTransfer
+		, const TArray<FBoneReference>& InInputBonesToTransfer
+		, const TArray<FBoneReference>& InOutputBonesToTransfer
 		, const TWeakObjectPtr<UNodeMappingContainer>& InNodeMappingContainer);
 
 	/** Complete mapping from skeleton to control rig bone index */
diff --git a/Engine/Plugins/Animation/ControlRig/Source/ControlRig/Public/Tools/ControlRigVariableMappings.h b/Engine/Plugins/Animation/ControlRig/Source/ControlRig/Public/Tools/ControlRigVariableMappings.h
--- a/Engine/Plugins/Animation/ControlRig/Source/ControlRig/Public/Tools/ControlRigVariableMappings.h
+++ b/Engine/Plugins/Animation/ControlRig/Source/ControlRig/Public/Tools/ControlRigVariableMappings.h
@@ -5,6 +5,8 @@
 #include "CoreMinimal.h"
 #include "Misc/Optional.h"
 #include "Animation/NamedValueArray.h"
+#include "Rigs/RigHierarchyDefines.h"
+#include "RigVMCore/RigVMExternalVariable.h"
 
 class UControlRig;
 class USkeletalMeshComponent;
@@ -13,10 +15,9 @@ class URigHierarchy;
 struct FBlendedCurve;
 struct FCompactPose;
 struct FBoneContainer;
-struct FRigVMExternalVariable;
 struct FRigControlElement;
 
-struct FControlRigVariableMappings
+struct CONTROLRIG_API FControlRigVariableMappings
 {
 	struct FControlRigCurveMapping
 	{
@@ -33,13 +34,78 @@ struct FControlRigVariableMappings
 
 	using FCurveMappings = UE::Anim::TNamedValueArray<FDefaultAllocator, FControlRigCurveMapping>;
 
+	struct FCustomPropertyData
+	{
+		enum class ECustomPropertyType : uint8
+		{
+			Variable = 0,
+			Control,
+			// --- ---
+			Num,
+			Invalid = Num
+		};
+
+		FCustomPropertyData() = default;
+		FCustomPropertyData(const FName& InTargetName, uint8* InTargetMemory, const FProperty* InProperty, const uint8* InSourceMemory, ECustomPropertyType InType)
+			: TargetName(InTargetName)
+			, TargetMemory(InTargetMemory)
+			, Property(InProperty)
+			, SourceMemory(InSourceMemory)
+			, Type(InType)
+		{}
+		FCustomPropertyData(const FName& InTargetName, uint8* InTargetMemory, const FProperty* InProperty, const uint8* InSourceMemory, ECustomPropertyType InType, ERigControlType InControlType)
+			: TargetName(InTargetName)
+			, TargetMemory(InTargetMemory)
+			, Property(InProperty)
+			, SourceMemory(InSourceMemory)
+			, Type(InType)
+			, ControlType(InControlType)
+		{}
+
+		const FName TargetName = NAME_None;
+		uint8* TargetMemory = nullptr;
+		const FProperty* Property = nullptr;
+		const uint8* SourceMemory = nullptr;
+		ECustomPropertyType Type = ECustomPropertyType::Invalid;
+		ERigControlType ControlType = ERigControlType::EulerTransform; // There is no invalid type for initialization
+	};
+
+	struct FCustomPropertyMappings
+	{
+		FCustomPropertyMappings() = default;
+		explicit FCustomPropertyMappings(int32 ReserveSize)
+		{
+			Mappings.Reserve(ReserveSize);
+		}
+
+		void Reset(int32 NewSize = 0) { Mappings.Reset(NewSize); }
+
+		void AddVariable(const FName& TargetName, uint8* TargetMemory, const FProperty* Property, const uint8* SourceMemory)
+		{
+			Mappings.Add(FCustomPropertyData(TargetName, TargetMemory, Property, SourceMemory, FCustomPropertyData::ECustomPropertyType::Variable));
+		}
+		void AddControl(ERigControlType InControlType, const FName& TargetName, uint8* TargetMemory, const FProperty* Property, const uint8* SourceMemory)
+		{
+			Mappings.Add(FCustomPropertyData(TargetName, TargetMemory, Property, SourceMemory, FCustomPropertyData::ECustomPropertyType::Control, InControlType));
+		}
+
+		const TArray<FCustomPropertyData>& GetMappings() const { return Mappings; }
+
+	private:
+		TArray<FCustomPropertyData> Mappings;
+	};
+
 	void InitializeProperties(const UClass* InSourceClass
 		, UObject* TargetInstance
 		, UClass* InTargetClass
 		, const TArray<FName>& SourcePropertyNames
 		, TArray<FName>& DestPropertyNames);
 
+	void InitializeCustomProperties(UControlRig* TargetControlRig, const FCustomPropertyMappings& CustomPropertyMapping);
+	void ResetCustomProperties(int32 NewSize = 0);
+
 	void PropagateInputProperties(const UObject* InSourceInstance, UControlRig* InTargetInstance, TArray<FName>& InDestPropertyNames);
+	void PropagateCustomInputProperties(UControlRig* InTargetControlRig);
 
 	void UpdateCurveInputs(UControlRig* InControlRig, const TMap<FName, FName>& InputMapping, const FBlendedCurve& CurveData);
 
@@ -64,6 +130,7 @@ struct FControlRigVariableMappings
 
 private:
 	using PropertyUpdateFunction = TFunction<void(const UObject*)>;
+	using CustomPropertyUpdateFunction = TFunction<void()>;
 
 	/*
 	 * Cached functions used to copy property values to avoid looking for them for every PropagateInputProperties calls.
@@ -71,7 +138,14 @@ private:
 	 */
 	void AddControlFunction(FRigControlElement* InControlElement, FProperty* InSourceProperty, URigHierarchy* InTargetHierarchy);
 	void AddUpdateFunction(PropertyUpdateFunction&& InFunction);
+	void AddCustomUpdateFunction(CustomPropertyUpdateFunction&& InFunction);
 	void AddVariableFunction(const FRigVMExternalVariable& InVariable, FProperty* InSourceProperty);
+
+	/**
+	 * Cached functions for non property based source data
+	 */
+	void AddCustomControlFunction(FRigControlElement* InControlElement, ERigControlType InControlType, const uint8* InSourcePropertyMemory, URigHierarchy* InTargetHierarchy);
+	void AddCustomVariableFunction(const FRigVMExternalVariable& InVariable, const FProperty* InSourceProperty, const uint8* InSourcePropertyMemory);
 	
 	/* This function is used to copy the property values going thru the hierarchy without any caching mechanism. */
 	void PropagateInputPropertiesNoCache(const UObject* InSourceInstance, const UControlRig* InTargetControlRig, URigHierarchy* InTargetHierarchy, TArray<FName>& InDestPropertyNames);
@@ -82,8 +156,12 @@ private:
 	TArray<FProperty*> DestProperties;
 
 	TArray<PropertyUpdateFunction> UpdateFunctions;
+	TArray<CustomPropertyUpdateFunction> CustomUpdateFunctions;
 	TArray<FRigVMExternalVariable> Variables;
 
+	TArray<const uint8*> SourcePropertiesMemory;
+	TArray<FName> DestPropertyNames;
+
 	// Bulk curves for I/O
 	FCurveMappings InputCurveMappings;
 	FCurveMappings OutputCurveMappings;
diff --git a/Engine/Plugins/Animation/ControlRig/Source/ControlRigDeveloper/Private/ControlRigIOMapping.cpp b/Engine/Plugins/Animation/ControlRig/Source/ControlRigDeveloper/Private/ControlRigIOMapping.cpp
--- a/Engine/Plugins/Animation/ControlRig/Source/ControlRigDeveloper/Private/ControlRigIOMapping.cpp
+++ b/Engine/Plugins/Animation/ControlRig/Source/ControlRigDeveloper/Private/ControlRigIOMapping.cpp
@@ -287,9 +287,14 @@ void FControlRigIOMapping::OnPropertyExposeCheckboxChanged(ECheckBoxState NewSta
 
 const TArray<FControlRigIOMapping::FControlsInfo>& FControlRigIOMapping::GetControls() const
 {
-	DECLARE_SCOPE_HIERARCHICAL_COUNTER_FUNC()
-
 	const UClass* ControlRigClass = GetTargetClass();
+	USkeleton* TargetSkeleton = GetTargetSkeleton();
+	return RigControlsData.GetControls(ControlRigClass, TargetSkeleton);
+}
+
+const TArray<FControlRigIOMapping::FControlsInfo>& FControlRigIOMapping::FRigControlsData::GetControls(const UClass* ControlRigClass, USkeleton* TargetSkeleton) const
+{
+	DECLARE_SCOPE_HIERARCHICAL_COUNTER_FUNC()
 
 	if (ControlsInfoClass != ControlRigClass)
 	{
@@ -298,15 +303,15 @@ const TArray<FControlRigIOMapping::FControlsInfo>& FControlRigIOMapping::GetCont
 
 		if (ControlRigClass)
 		{
-			if (USkeleton* Skeleton = GetTargetSkeleton())
+			if (TargetSkeleton != nullptr)
 			{
 				UControlRig* TemplateRig = NewObject<UControlRig>(GetTransientPackage(), ControlRigClass, NAME_None, RF_Transient);
 				TemplateRig->SetObjectBinding(MakeShared<FControlRigObjectBinding>());
-				TemplateRig->GetObjectBinding()->BindToObject(Skeleton);
+				TemplateRig->GetObjectBinding()->BindToObject(TargetSkeleton);
 				TemplateRig->GetDataSourceRegistry()->RegisterDataSource(UControlRig::OwnerComponent, TemplateRig->GetObjectBinding()->GetBoundObject());
 
 				TemplateRig->Initialize();
-				TemplateRig->SetBoneInitialTransformsFromRefSkeleton(Skeleton->GetReferenceSkeleton());
+				TemplateRig->SetBoneInitialTransformsFromRefSkeleton(TargetSkeleton->GetReferenceSkeleton());
 				{
 					// Empty event queue to evaluate only construction
 					TGuardValue<TArray<FName>> EventQueueGuard(TemplateRig->EventQueue, {});
diff --git a/Engine/Plugins/Animation/ControlRig/Source/ControlRigDeveloper/Public/ControlRigIOMapping.h b/Engine/Plugins/Animation/ControlRig/Source/ControlRigDeveloper/Public/ControlRigIOMapping.h
--- a/Engine/Plugins/Animation/ControlRig/Source/ControlRigDeveloper/Public/ControlRigIOMapping.h
+++ b/Engine/Plugins/Animation/ControlRig/Source/ControlRigDeveloper/Public/ControlRigIOMapping.h
@@ -14,7 +14,6 @@ class UClass;
 class USkeleton;
 class UControlRig;
 
-
 struct CONTROLRIGDEVELOPER_API FControlRigIOMapping : public TSharedFromThis<FControlRigIOMapping>
 {
 	DECLARE_DELEGATE_RetVal(USkeleton*, FOnGetTargetSkeleton);
@@ -26,7 +25,7 @@ struct CONTROLRIGDEVELOPER_API FControlRigIOMapping : public TSharedFromThis<FCo
 	FControlRigIOMapping(TMap<FName, FName>& InInputMapping, TMap<FName, FName>& InOutputMapping, TArray<FOptionalPinFromProperty>& InCustomPinProperties);
 	virtual ~FControlRigIOMapping();
 
-	struct FControlsInfo
+	struct CONTROLRIGDEVELOPER_API FControlsInfo
 	{
 		FName Name;
 		FString DisplayName;
@@ -34,7 +33,20 @@ struct CONTROLRIGDEVELOPER_API FControlRigIOMapping : public TSharedFromThis<FCo
 		ERigControlType ControlType;
 		FString DefaultValue;
 	};
-	
+
+	// Helper enabling quering controls in editor from a Control Rig class
+	struct CONTROLRIGDEVELOPER_API FRigControlsData
+	{
+		FRigControlsData() = default;
+
+		const TArray<FControlsInfo>& GetControls(const UClass* ControlRigClass, USkeleton* TargetSkeleton) const;
+
+	private:
+		// Controls data
+		mutable const UClass* ControlsInfoClass = nullptr;
+		mutable TArray<FControlsInfo> ControlsInfo;
+	};
+
 	bool CreateVariableMappingWidget(IDetailLayoutBuilder& DetailBuilder);
 
 	const TMap<FName, FRigVMExternalVariable>& GetInputVariables() const
@@ -86,8 +98,7 @@ private:
 	TArray<FOptionalPinFromProperty>& CustomPinProperties;
 
 	// Controls data
-	mutable const UClass* ControlsInfoClass = nullptr;
-	mutable TArray<FControlsInfo> ControlsInfo;
+	FRigControlsData RigControlsData;
 
 	// property related things
 	static void GetVariables(const UClass* TargetClass, bool bInput, TMap<FName, FRigVMExternalVariable>& OutParameters);
diff --git a/Engine/Plugins/Experimental/AnimNext/Source/AnimNext/Private/GenerationTools.cpp b/Engine/Plugins/Experimental/AnimNext/Source/AnimNext/Private/GenerationTools.cpp
--- a/Engine/Plugins/Experimental/AnimNext/Source/AnimNext/Private/GenerationTools.cpp
+++ b/Engine/Plugins/Experimental/AnimNext/Source/AnimNext/Private/GenerationTools.cpp
@@ -18,6 +18,7 @@
 #include "AssetRegistry/AssetData.h"
 #include "AssetRegistry/AssetRegistryModule.h"
 #include "AnimNextStats.h"
+#include "Animation/AttributesContainer.h"
 
 DEFINE_STAT(STAT_AnimNext_GenerateReferencePose);
 DEFINE_STAT(STAT_AnimNext_RemapPose_FromAnimBP);
@@ -634,9 +635,10 @@ void FGenerationTools::RemapAttributes(
 	}
 }
 
-void FGenerationTools::RemapAttributes(
+template<typename BoneIndexType, typename InAllocator>
+static void RemapAttributesImpl(
 	const FLODPose& LODPose,
-	const UE::Anim::FHeapAttributeContainer& InAttributes,
+	const UE::Anim::TAttributeContainer<BoneIndexType, InAllocator>& InAttributes,
 	FPoseContext& OutPose)
 {
 	const TArrayView<const FBoneIndexType> LODBoneIndexToMeshBoneIndexMap = LODPose.GetLODBoneIndexToMeshBoneIndexMap();
@@ -648,7 +650,7 @@ void FGenerationTools::RemapAttributes(
 		const int32 TypeIndex = InAttributes.FindTypeIndex(ScriptStruct);
 		if (TypeIndex != INDEX_NONE)
 		{
-			const TArray<UE::Anim::TWrappedAttribute<FDefaultAllocator>, FDefaultAllocator>& SourceValues = InAttributes.GetValues(TypeIndex);
+			const TArray<UE::Anim::TWrappedAttribute<InAllocator>, FDefaultAllocator>& SourceValues = InAttributes.GetValues(TypeIndex);
 			const TArray<UE::Anim::FAttributeId, FDefaultAllocator>& AttributeIds = InAttributes.GetKeys(TypeIndex);
 
 			// Try and remap all the source attributes to their respective new bone indices
@@ -670,7 +672,7 @@ void FGenerationTools::RemapAttributes(
 				{
 					const UE::Anim::FAttributeId NewInfo(AttributeId.GetName(), CompactBoneIndex);
 					uint8* NewAttribute = OutPose.CustomAttributes.FindOrAdd(ScriptStruct, NewInfo);
-					ScriptStruct->CopyScriptStruct(NewAttribute, SourceValues[EntryIndex].GetPtr<void>());
+					ScriptStruct->CopyScriptStruct(NewAttribute, SourceValues[EntryIndex].template GetPtr<void>());
 				}
 				else
 				{
@@ -681,6 +683,87 @@ void FGenerationTools::RemapAttributes(
 	}
 }
 
+void FGenerationTools::RemapAttributes(
+	const FLODPose& LODPose,
+	const UE::Anim::FHeapAttributeContainer& InAttributes,
+	FPoseContext& OutPose)
+{
+	RemapAttributesImpl(LODPose, InAttributes, OutPose);
+}
+
+void FGenerationTools::RemapAttributes(
+	const FLODPose& LODPose,
+	const UE::Anim::FStackAttributeContainer& InAttributes,
+	FPoseContext& OutPose)
+{
+	RemapAttributesImpl(LODPose, InAttributes, OutPose);
+}
+
+template<typename InBoneIndexType, typename InAllocator, typename OutBoneIndexType, typename OutAllocator>
+static void RemapCompactPoseAttributesToLODPoseAttributes(
+	const FBoneContainer& InBoneContainer,
+	const UE::Anim::TAttributeContainer<InBoneIndexType, InAllocator>& InAttributes,
+	const TArrayView<const FBoneIndexType>& SkeletonBoneIndexToLODBoneIndexMap,
+	UE::Anim::TAttributeContainer<OutBoneIndexType, OutAllocator>& OutAttributes
+	)
+{
+	for (const TWeakObjectPtr<UScriptStruct> WeakScriptStruct : InAttributes.GetUniqueTypes())
+	{
+		const UScriptStruct* ScriptStruct = WeakScriptStruct.Get();
+		const int32 TypeIndex = InAttributes.FindTypeIndex(ScriptStruct);
+		if (TypeIndex != INDEX_NONE)
+		{
+			const TArray<UE::Anim::TWrappedAttribute<InAllocator>, FDefaultAllocator>& SourceValues = InAttributes.GetValues(TypeIndex);
+			const TArray<UE::Anim::FAttributeId, FDefaultAllocator>& AttributeIds = InAttributes.GetKeys(TypeIndex);
+
+			// Try and remap all the source attributes to their respective new bone indices
+			for (int32 EntryIndex = 0; EntryIndex < AttributeIds.Num(); ++EntryIndex)
+			{
+				const UE::Anim::FAttributeId& AttributeId = AttributeIds[EntryIndex];
+
+				// Obtain the compact pose bone index from the attribute
+				const FCompactPoseBoneIndex CompactBoneIndex(AttributeId.GetIndex());
+				ensure(CompactBoneIndex.IsValid());
+
+				// Remap the compact pose bone index to the skeleton index
+				const FSkeletonPoseBoneIndex SkeletonBoneIndex = InBoneContainer.GetSkeletonPoseIndexFromCompactPoseIndex(CompactBoneIndex);
+				ensure(SkeletonBoneIndex.IsValid());	// We expect the skeletal mesh bone to map to a valid skeleton bone
+				
+				// Remap our skeleton bone index into the LOD Pose bone index we output for
+				const FBoneIndexType& LODBoneIndex = SkeletonBoneIndexToLODBoneIndexMap[SkeletonBoneIndex.GetInt()];
+				if (ensure(LODBoneIndex != INDEX_NONE))	// We expect the skeleton bone to map to a valid LOD Pose bone index
+				{
+					const UE::Anim::FAttributeId NewInfo(AttributeId.GetName(), LODBoneIndex, AttributeId.GetNamespace());
+					uint8* NewAttribute = OutAttributes.FindOrAdd(ScriptStruct, NewInfo);
+					ScriptStruct->CopyScriptStruct(NewAttribute, SourceValues[EntryIndex].template GetPtr<void>());
+				}
+				else
+				{
+					// This bone is part of the LOD but isn't part of the required bones
+				}
+			}
+		}
+	}
+}
+
+void FGenerationTools::RemapAttributes(
+	const FPoseContext& OutPose,
+	const FLODPose& LODPose,
+	UE::Anim::FHeapAttributeContainer& OutAttributes
+	)
+{
+	RemapCompactPoseAttributesToLODPoseAttributes(OutPose.Pose.GetBoneContainer(), OutPose.CustomAttributes, LODPose.GetSkeletonBoneIndexToLODBoneIndexMap(), OutAttributes);
+}
+
+void FGenerationTools::RemapAttributes(
+	const FPoseContext& OutPose,
+	const FLODPose& LODPose,
+	UE::Anim::FStackAttributeContainer& OutAttributes
+	)
+{
+	RemapCompactPoseAttributesToLODPoseAttributes(OutPose.Pose.GetBoneContainer(), OutPose.CustomAttributes, LODPose.GetSkeletonBoneIndexToLODBoneIndexMap(), OutAttributes);
+}
+
 void FGenerationTools::ConvertLocalSpaceToComponentSpace(
 												 TConstArrayView<FBoneIndexType> InMeshBoneIndexToParentMeshBoneIndexMap,
 												 TConstArrayView<FTransform> InBoneSpaceTransforms,
diff --git a/Engine/Plugins/Experimental/AnimNext/Source/AnimNext/Public/GenerationTools.h b/Engine/Plugins/Experimental/AnimNext/Source/AnimNext/Public/GenerationTools.h
--- a/Engine/Plugins/Experimental/AnimNext/Source/AnimNext/Public/GenerationTools.h
+++ b/Engine/Plugins/Experimental/AnimNext/Source/AnimNext/Public/GenerationTools.h
@@ -119,6 +119,26 @@ public:
 		const UE::Anim::FHeapAttributeContainer& InAttributes,
 		FPoseContext& OutPose);
 
+	// Converts AnimNext attributes to AnimBP attributes
+	static void RemapAttributes(
+		const FLODPose& LODPose,
+		const UE::Anim::FStackAttributeContainer& InAttributes,
+		FPoseContext& OutPose);
+
+	// Converts AnimBP attributes to AnimNext attributes
+	static void RemapAttributes(
+		const FPoseContext& OutPose,
+		const FLODPose& LODPose,
+		UE::Anim::FHeapAttributeContainer& OutAttributes
+	);
+
+	// Converts AnimBP attributes to AnimNext attributes
+	static void RemapAttributes(
+		const FPoseContext& OutPose,
+		const FLODPose& LODPose,
+		UE::Anim::FStackAttributeContainer& OutAttributes
+	);
+
 	// Converts a local space to component space buffer given a number of required bones
 	static void ConvertLocalSpaceToComponentSpace(TConstArrayView<FBoneIndexType> InParentIndices, TConstArrayView<FTransform> InBoneSpaceTransforms, TConstArrayView<FBoneIndexType> InRequiredBoneIndices, TArrayView<FTransform> OutComponentSpaceTransforms);
 };
diff --git a/Engine/Plugins/Experimental/AnimNext/Source/AnimNext/Public/LODPose.h b/Engine/Plugins/Experimental/AnimNext/Source/AnimNext/Public/LODPose.h
--- a/Engine/Plugins/Experimental/AnimNext/Source/AnimNext/Public/LODPose.h
+++ b/Engine/Plugins/Experimental/AnimNext/Source/AnimNext/Public/LODPose.h
@@ -139,6 +139,19 @@ struct FLODPose
 		}
 	}
 
+	const TArrayView<const FBoneIndexType> GetMeshBoneIndexToLODBoneIndexMap() const
+	{
+		if (RefPose != nullptr)
+		{
+			return RefPose->GetMeshBoneIndexToLODBoneIndexMap();
+		}
+		else
+		{
+			return TArrayView<const FBoneIndexType>();
+		}
+
+	}
+
 	const TArrayView<const FBoneIndexType> GetSkeletonBoneIndexToLODBoneIndexMap() const
 	{
 		if (RefPose != nullptr)
diff --git a/Engine/Plugins/Experimental/AnimNext/Source/AnimNextEditor/Private/AnimNextEditorModule.cpp b/Engine/Plugins/Experimental/AnimNext/Source/AnimNextEditor/Private/AnimNextEditorModule.cpp
--- a/Engine/Plugins/Experimental/AnimNext/Source/AnimNextEditor/Private/AnimNextEditorModule.cpp
+++ b/Engine/Plugins/Experimental/AnimNext/Source/AnimNextEditor/Private/AnimNextEditorModule.cpp
@@ -441,6 +441,16 @@ void FAnimNextEditorModule::UnregisterAssetCompilationHandler(const FTopLevelAss
 	AssetCompilationHandlerFactories.Remove(InClassPath);
 }
 
+FDelegateHandle FAnimNextEditorModule::RegisterNodeDblClickHandler(const FNodeDblClickNotificationDelegate& InNodeDblClickNotificationDelegate)
+{
+	return OnNodeDblClickHandlerMulticast.Add(InNodeDblClickNotificationDelegate);
+}
+
+void FAnimNextEditorModule::UnregisterNodeDblClickHandler(const FDelegateHandle& InDelegateHandle)
+{
+	OnNodeDblClickHandlerMulticast.Remove(InDelegateHandle);
+}
+
 void FAnimNextEditorModule::CollectGraphMenuActions(const TWeakPtr<UE::Workspace::IWorkspaceEditor>& WorkspaceEditor, FGraphContextMenuBuilder& InContexMenuBuilder, const FActionMenuContextData& InActionMenuContextData)
 {
 	if (OnCollectGraphMenuActionsDelegateImpl.IsBound())
@@ -604,7 +614,7 @@ void FAnimNextEditorModule::RegisterWorkspaceDocumentTypes(Workspace::IWorkspace
 	AnimNextModuleDocumentArgs.DocumentEditorMode = UAnimNextWorkspaceEditorMode::EM_AnimNextWorkspace;
 
 	WorkspaceEditorModule.RegisterObjectDocumentType(FTopLevelAssetPath(TEXT("/Script/AnimNext.AnimNextModule")), AnimNextModuleDocumentArgs);
-	
+
 	// --- AnimNextEdGraph ---
 	Workspace::FGraphDocumentWidgetArgs GraphArgs;
 	GraphArgs.SpawnLocation = Workspace::WorkspaceTabs::TopMiddleDocumentArea;
@@ -642,10 +652,8 @@ void FAnimNextEditorModule::RegisterWorkspaceDocumentTypes(Workspace::IWorkspace
 
 		RigVMEdGraph->GetController()->SetCommentTextByName(CommentBeingChanged->GetFName(), NewText.ToString(), CommentBeingChanged->FontSize, CommentBeingChanged->bCommentBubbleVisible, CommentBeingChanged->bColorCommentBubble, true, true);
 	});
-	GraphArgs.OnCanDeleteSelectedNodes = Workspace::FOnCanPerformActionOnSelectedNodes::CreateLambda([](const Workspace::FWorkspaceEditorContext& InContext, const FGraphPanelSelectionSet& InSelectedNodes)
+	GraphArgs.OnCanDeleteSelectedNodes = Workspace::FOnCanPerformActionOnSelectedNodes::CreateLambda([this](const Workspace::FWorkspaceEditorContext& InContext, const FGraphPanelSelectionSet& InSelectedNodes)
 	{
-		bool bCanUserDeleteNode = false;
-
 		if(InSelectedNodes.Num() > 0)
 		{
 			for(UObject* NodeObject : InSelectedNodes)
@@ -654,13 +662,11 @@ void FAnimNextEditorModule::RegisterWorkspaceDocumentTypes(Workspace::IWorkspace
 				const UEdGraphNode* Node = Cast<UEdGraphNode>(NodeObject);
 				if(Node && Node->CanUserDeleteNode())
 				{
-					bCanUserDeleteNode = true;
-					break;
+					return true;
 				}
 			}
 		}
-
-		return bCanUserDeleteNode;
+		return false;
 	});
 	GraphArgs.OnDeleteSelectedNodes = Workspace::FOnPerformActionOnSelectedNodes::CreateLambda([](const Workspace::FWorkspaceEditorContext& InContext, const FGraphPanelSelectionSet& InSelectedNodes)
 	{
@@ -734,16 +740,13 @@ void FAnimNextEditorModule::RegisterWorkspaceDocumentTypes(Workspace::IWorkspace
 		Controller->RemoveNodes(NodesToRemove, true);
 		Controller->CloseUndoBracket();
 	});
-	GraphArgs.OnCanCopySelectedNodes = Workspace::FOnCanPerformActionOnSelectedNodes::CreateLambda([](const Workspace::FWorkspaceEditorContext& InContext, const FGraphPanelSelectionSet& InSelectedNodes)
+	GraphArgs.OnCanCopySelectedNodes = Workspace::FOnCanPerformActionOnSelectedNodes::CreateLambda([this](const Workspace::FWorkspaceEditorContext& InContext, const FGraphPanelSelectionSet& InSelectedNodes)
 	{
-		bool bCanUserCopyNode = false;
-
 		if(InSelectedNodes.Num() > 0)
 		{
-			bCanUserCopyNode = true;
+			return true;
 		}
-
-		return bCanUserCopyNode;
+		return false;
 	});
 	GraphArgs.OnCopySelectedNodes = Workspace::FOnPerformActionOnSelectedNodes::CreateLambda([](const Workspace::FWorkspaceEditorContext& InContext, const FGraphPanelSelectionSet& InSelectedNodes)
 	{
@@ -861,7 +864,7 @@ void FAnimNextEditorModule::RegisterWorkspaceDocumentTypes(Workspace::IWorkspace
 			return bCanUserCopyNode;
 		});
 	UE::Workspace::FOnPasteNodes& OnPasteNodes = GraphArgs.OnPasteNodes;
-	GraphArgs.OnDuplicateSelectedNodes = Workspace::FOnDuplicateSelectedNodes::CreateLambda([OnCopySelectedNodes, OnPasteNodes](const Workspace::FWorkspaceEditorContext& InContext, const FVector2D& InPasteLocation, const FGraphPanelSelectionSet& InSelectedNodes)
+	GraphArgs.OnDuplicateSelectedNodes = Workspace::FOnDuplicateSelectedNodes::CreateLambda([this, OnCopySelectedNodes, OnPasteNodes](const Workspace::FWorkspaceEditorContext& InContext, const FVector2D& InPasteLocation, const FGraphPanelSelectionSet& InSelectedNodes)
 		{
 			if (InSelectedNodes.IsEmpty())
 			{
@@ -888,7 +891,7 @@ void FAnimNextEditorModule::RegisterWorkspaceDocumentTypes(Workspace::IWorkspace
 				Controller->CloseUndoBracket();
 			}
 		});
-	GraphArgs.OnGraphSelectionChanged = Workspace::FOnGraphSelectionChanged::CreateLambda([](const Workspace::FWorkspaceEditorContext& InContext, const FGraphPanelSelectionSet& NewSelection)
+	GraphArgs.OnGraphSelectionChanged = Workspace::FOnGraphSelectionChanged::CreateLambda([this](const Workspace::FWorkspaceEditorContext& InContext, const FGraphPanelSelectionSet& NewSelection)
 	{
 		URigVMEdGraph* RigVMEdGraph = Cast<URigVMEdGraph>(InContext.Object);
 		if (RigVMEdGraph == nullptr)
@@ -919,12 +922,17 @@ void FAnimNextEditorModule::RegisterWorkspaceDocumentTypes(Workspace::IWorkspace
 
 		InContext.WorkspaceEditor->SetDetailsObjects(NewSelection.Array());
 	});
-	GraphArgs.OnNodeDoubleClicked = Workspace::FOnNodeDoubleClicked::CreateLambda([](const Workspace::FWorkspaceEditorContext& InContext, const UEdGraphNode* InNode)
+	GraphArgs.OnNodeDoubleClicked = Workspace::FOnNodeDoubleClicked::CreateLambda([this](const Workspace::FWorkspaceEditorContext& InContext, const UEdGraphNode* InNode)
 	{
-		if (const URigVMEdGraphNode* RigVMEdGraphNode = Cast<URigVMEdGraphNode>(InNode))
+		if (const UAnimNextEdGraphNode* RigVMEdGraphNode = Cast<UAnimNextEdGraphNode>(InNode))
 		{
 			const URigVMNode* ModelNode = RigVMEdGraphNode->GetModelNode();
 
+			if (OnNodeDblClickHandlerMulticast.IsBound())
+			{
+				OnNodeDblClickHandlerMulticast.Broadcast(InContext, InNode);
+			}
+
 			if (const URigVMLibraryNode* LibraryNode = Cast<URigVMLibraryNode>(ModelNode))
 			{
 				URigVMGraph* ContainedGraph = LibraryNode->GetContainedGraph();
diff --git a/Engine/Plugins/Experimental/AnimNext/Source/AnimNextEditor/Private/AnimNextEditorModule.h b/Engine/Plugins/Experimental/AnimNext/Source/AnimNextEditor/Private/AnimNextEditorModule.h
--- a/Engine/Plugins/Experimental/AnimNext/Source/AnimNextEditor/Private/AnimNextEditorModule.h
+++ b/Engine/Plugins/Experimental/AnimNext/Source/AnimNextEditor/Private/AnimNextEditorModule.h
@@ -41,6 +41,8 @@ private:
 	virtual void UnregisterGraphMenuActionsProvider(const FDelegateHandle& InDelegateHandle) override;
 	virtual void RegisterAssetCompilationHandler(const FTopLevelAssetPath& InClassPath, FAssetCompilationHandlerFactoryDelegate InAssetCompilationHandlerFactory) override;
 	virtual void UnregisterAssetCompilationHandler(const FTopLevelAssetPath& InClassPath) override;
+	virtual FDelegateHandle RegisterNodeDblClickHandler(const FNodeDblClickNotificationDelegate& InNodeDblClickNotificationDelegate) override;
+	virtual void UnregisterNodeDblClickHandler(const FDelegateHandle& InDelegateHandle) override;
 
 private:
 	void CollectGraphMenuActions(const TWeakPtr<UE::Workspace::IWorkspaceEditor>& WorkspaceEditor, FGraphContextMenuBuilder& InContexMenuBuilder, const FActionMenuContextData& InActionMenuContextData);
@@ -60,6 +62,7 @@ private:
 	TArray<FTopLevelAssetPath> SupportedAssetClasses;
 
 	FCollectGraphMenuActionsMulticast OnCollectGraphMenuActionsDelegateImpl;
+	FNodeDblClickNotificationMulticast OnNodeDblClickHandlerMulticast;
 
 	TMap<FTopLevelAssetPath, FAssetCompilationHandlerFactoryDelegate> AssetCompilationHandlerFactories;
 
diff --git a/Engine/Plugins/Experimental/AnimNext/Source/AnimNextEditor/Public/IAnimNextEditorModule.h b/Engine/Plugins/Experimental/AnimNext/Source/AnimNextEditor/Public/IAnimNextEditorModule.h
--- a/Engine/Plugins/Experimental/AnimNext/Source/AnimNextEditor/Public/IAnimNextEditorModule.h
+++ b/Engine/Plugins/Experimental/AnimNext/Source/AnimNextEditor/Public/IAnimNextEditorModule.h
@@ -8,14 +8,18 @@
 
 class SWidget;
 class IMessageLogListing;
+class UEdGraphNode;
 struct FTopLevelAssetPath;
 struct FGraphContextMenuBuilder;
 
 namespace UE::Workspace
 {
 	class IWorkspaceEditor;
+	struct FWorkspaceEditorContext;
 }
 
+typedef TSet<class UObject*> FGraphPanelSelectionSet;
+
 namespace UE::AnimNext::Editor
 {
 
@@ -36,6 +40,9 @@ public:
 	DECLARE_MULTICAST_DELEGATE_ThreeParams(FCollectGraphMenuActionsMulticast, const TWeakPtr<UE::Workspace::IWorkspaceEditor>& /*IWorkspaceEditor*/, FGraphContextMenuBuilder& /*InContexMenuBuilder*/, const FActionMenuContextData& /*ActionMenuContextData*/);
 	typedef FCollectGraphMenuActionsMulticast::FDelegate FOnCollectGraphMenuActionsDelegate;
 
+	DECLARE_MULTICAST_DELEGATE_TwoParams(FNodeDblClickNotificationMulticast, const UE::Workspace::FWorkspaceEditorContext& /*InWorkspaceEditorContext*/, const UEdGraphNode* /*InNode*/);
+	typedef FNodeDblClickNotificationMulticast::FDelegate FNodeDblClickNotificationDelegate;
+
 	// Register a valid fragment type name to be used with parameter UOLs
 	// @param InLocatorFragmentEditorName The name of the locator fragment editor
 	virtual void RegisterLocatorFragmentEditorType(FName InLocatorFragmentEditorName) = 0;
@@ -69,6 +76,14 @@ public:
 	// Unregister an asset compilation handler
 	// @param InClassPath The path of the asset's class
 	virtual void UnregisterAssetCompilationHandler(const FTopLevelAssetPath& InClassPath) = 0;
+
+	// Register a node double click notifications handler
+	// @param InGraphNodeActionsHandler Handler interface to register
+	virtual FDelegateHandle RegisterNodeDblClickHandler(const FNodeDblClickNotificationDelegate& InNodeDblClickNotificationDelegate) = 0;
+
+	// Unregister a node double click notifications handler
+	// @param InDelegateHandle Handle of the delegate to remove
+	virtual void UnregisterNodeDblClickHandler(const FDelegateHandle& InDelegateHandle) = 0;
 };
 
 }
diff --git a/Engine/Plugins/Experimental/AnimNextAnimGraph/Source/AnimNextAnimGraph/Private/Graph/RigDecorator_AnimNextCppTrait.cpp b/Engine/Plugins/Experimental/AnimNextAnimGraph/Source/AnimNextAnimGraph/Private/Graph/RigDecorator_AnimNextCppTrait.cpp
--- a/Engine/Plugins/Experimental/AnimNextAnimGraph/Source/AnimNextAnimGraph/Private/Graph/RigDecorator_AnimNextCppTrait.cpp
+++ b/Engine/Plugins/Experimental/AnimNextAnimGraph/Source/AnimNextAnimGraph/Private/Graph/RigDecorator_AnimNextCppTrait.cpp
@@ -40,6 +40,14 @@ void FRigDecorator_AnimNextCppDecorator::GetProgrammaticPins(URigVMController* I
 	};
 
 	const int32 StartPinIndex = OutPinArray.Num();
+
+	// Give the Trait an opportunity to generate its own pins
+	const UE::AnimNext::FTrait* Trait = GetTrait();
+	check(Trait != nullptr);
+	FAnimNextTraitSharedData* TraitSharedData = (FAnimNextTraitSharedData*)DefaultValueMemoryScope.GetStructMemory();
+	Trait->GetProgrammaticPins(TraitSharedData, InController, InParentPinIndex, InDefaultValue, OutPinArray);
+	
+	// Generate the shared data struct pins
 	OutPinArray.AddPins(DecoratorSharedDataStruct, InController, ERigVMPinDirection::Invalid, InParentPinIndex, DefaultValueTypeGetter, DefaultValueMemoryScope.GetStructMemory(), true);
 
 	for (int32 PinIndex = StartPinIndex; PinIndex < OutPinArray.Num(); ++PinIndex)
diff --git a/Engine/Plugins/Experimental/AnimNextAnimGraph/Source/AnimNextAnimGraph/Private/Graph/RigUnit_AnimNextGraphEvaluator.cpp b/Engine/Plugins/Experimental/AnimNextAnimGraph/Source/AnimNextAnimGraph/Private/Graph/RigUnit_AnimNextGraphEvaluator.cpp
--- a/Engine/Plugins/Experimental/AnimNextAnimGraph/Source/AnimNextAnimGraph/Private/Graph/RigUnit_AnimNextGraphEvaluator.cpp
+++ b/Engine/Plugins/Experimental/AnimNextAnimGraph/Source/AnimNextAnimGraph/Private/Graph/RigUnit_AnimNextGraphEvaluator.cpp
@@ -43,6 +43,12 @@ void FRigUnit_AnimNextGraphEvaluator::StaticExecute(FRigVMExtendedExecuteContext
 			continue;
 		}
 
+		if (!Handle.IsOffsetValid())
+		{
+			// This handle isn't valid
+			continue;
+		}
+
 		if (bIsFrozen && Handle.CanFreeze())
 		{
 			// This handle can freeze and we are frozen, no need to update it
diff --git a/Engine/Plugins/Experimental/AnimNextAnimGraph/Source/AnimNextAnimGraph/Private/TraitCore/ExecutionContext.cpp b/Engine/Plugins/Experimental/AnimNextAnimGraph/Source/AnimNextAnimGraph/Private/TraitCore/ExecutionContext.cpp
--- a/Engine/Plugins/Experimental/AnimNextAnimGraph/Source/AnimNextAnimGraph/Private/TraitCore/ExecutionContext.cpp
+++ b/Engine/Plugins/Experimental/AnimNextAnimGraph/Source/AnimNextAnimGraph/Private/TraitCore/ExecutionContext.cpp
@@ -566,7 +566,7 @@ namespace UE::AnimNext
 	}
 
 #if UE_ENABLE_DEBUG_DRAWING
-	FRigVMDrawInterface* FExecutionContext::GetDebugDrawInterface()
+	FRigVMDrawInterface* FExecutionContext::GetDebugDrawInterface() const
 	{
 		check(IsBound());
 		return RootGraphInstance->GetModuleInstance()->GetDebugDrawInterface();
diff --git a/Engine/Plugins/Experimental/AnimNextAnimGraph/Source/AnimNextAnimGraph/Private/TraitCore/Trait.cpp b/Engine/Plugins/Experimental/AnimNextAnimGraph/Source/AnimNextAnimGraph/Private/TraitCore/Trait.cpp
--- a/Engine/Plugins/Experimental/AnimNextAnimGraph/Source/AnimNextAnimGraph/Private/TraitCore/Trait.cpp
+++ b/Engine/Plugins/Experimental/AnimNextAnimGraph/Source/AnimNextAnimGraph/Private/TraitCore/Trait.cpp
@@ -125,10 +125,14 @@ namespace UE::AnimNext
 		}
 	}
 
-	TArray<FLatentPropertyMetadata> FTrait::GetLatentPropertyHandles(
+	uint32 FTrait::GetLatentPropertyHandles(
+		const FAnimNextTraitSharedData* InSharedData,
+		TArray<FLatentPropertyMetadata>& OutLatentPropertyHandles,
 		bool bFilterEditorOnly,
 		const TFunction<uint16(FName PropertyName)>& GetTraitLatentPropertyIndex) const
 	{
+		uint32 NumHandlesAdded = 0;
+
 		const UStruct* BaseStruct = GetTraitSharedDataStruct();
 
 		// The property linked list on UScriptStruct iterates over the properties starting in the derived type
@@ -143,8 +147,6 @@ namespace UE::AnimNext
 		}
 		while (BaseStruct != nullptr);
 
-		TArray<FLatentPropertyMetadata> LatentPropertyHandles;
-
 		// Gather our latent properties from base to most derived
 		for (auto It = StructHierarchy.rbegin(); It != StructHierarchy.rend(); ++It)
 		{
@@ -180,11 +182,12 @@ namespace UE::AnimNext
 				// Always false for now, we don't support freezing yet
 				Metadata.bCanFreeze = false;
 
-				LatentPropertyHandles.Add(Metadata);
+				OutLatentPropertyHandles.Add(Metadata);
+				NumHandlesAdded++;
 			}
 		}
 
-		return LatentPropertyHandles;
+		return NumHandlesAdded;
 	}
 #endif
 
diff --git a/Engine/Plugins/Experimental/AnimNextAnimGraph/Source/AnimNextAnimGraph/Private/TraitCore/TraitWriter.cpp b/Engine/Plugins/Experimental/AnimNextAnimGraph/Source/AnimNextAnimGraph/Private/TraitCore/TraitWriter.cpp
--- a/Engine/Plugins/Experimental/AnimNextAnimGraph/Source/AnimNextAnimGraph/Private/TraitCore/TraitWriter.cpp
+++ b/Engine/Plugins/Experimental/AnimNextAnimGraph/Source/AnimNextAnimGraph/Private/TraitCore/TraitWriter.cpp
@@ -201,9 +201,11 @@ namespace UE::AnimNext
 				return GetTraitLatentPropertyIndex(TraitIndex, PropertyName);
 			};
 
-			TArray<FLatentPropertyMetadata> LatentProperties = Trait->GetLatentPropertyHandles(IsFilterEditorOnly(), GetTraitLatentPropertyIndexAt);
+			TArray<FLatentPropertyMetadata> LatentProperties;
+			LatentProperties.Reserve(32);
+			FAnimNextTraitSharedData* SharedData = TraitTemplates[TraitIndex].GetTraitDescription(*NodeDesc);
+			uint32 NumLatentProperties = Trait->GetLatentPropertyHandles(SharedData, LatentProperties, IsFilterEditorOnly(), GetTraitLatentPropertyIndexAt);
 
-			uint32 NumLatentProperties = LatentProperties.Num();
 			*this << NumLatentProperties;
 
 			for (FLatentPropertyMetadata& Metadata : LatentProperties)
diff --git a/Engine/Plugins/Experimental/AnimNextAnimGraph/Source/AnimNextAnimGraph/Public/EvaluationVM/EvaluationVM.h b/Engine/Plugins/Experimental/AnimNextAnimGraph/Source/AnimNextAnimGraph/Public/EvaluationVM/EvaluationVM.h
--- a/Engine/Plugins/Experimental/AnimNextAnimGraph/Source/AnimNextAnimGraph/Public/EvaluationVM/EvaluationVM.h
+++ b/Engine/Plugins/Experimental/AnimNextAnimGraph/Source/AnimNextAnimGraph/Public/EvaluationVM/EvaluationVM.h
@@ -187,6 +187,10 @@ namespace UE::AnimNext
 		// Returns an uninitialized keyframe with memory pre-allocated
 		[[nodiscard]] FKeyframeState MakeUninitializedKeyframe(bool bAdditiveKeyframe) const;
 
+		int32 GetCurrentLOD() const
+		{
+			return CurrentLOD;
+		}
 	private:
 		// Disallow copy
 		FEvaluationVM(const FEvaluationVM&) = delete;
diff --git a/Engine/Plugins/Experimental/AnimNextAnimGraph/Source/AnimNextAnimGraph/Public/TraitCore/ExecutionContext.h b/Engine/Plugins/Experimental/AnimNextAnimGraph/Source/AnimNextAnimGraph/Public/TraitCore/ExecutionContext.h
--- a/Engine/Plugins/Experimental/AnimNextAnimGraph/Source/AnimNextAnimGraph/Public/TraitCore/ExecutionContext.h
+++ b/Engine/Plugins/Experimental/AnimNextAnimGraph/Source/AnimNextAnimGraph/Public/TraitCore/ExecutionContext.h
@@ -205,7 +205,7 @@ namespace UE::AnimNext
 
 #if UE_ENABLE_DEBUG_DRAWING
 		// Get the debug draw interface
-		FRigVMDrawInterface* GetDebugDrawInterface();
+		FRigVMDrawInterface* GetDebugDrawInterface() const;
 #endif
 	private:
 		// No copy or move
diff --git a/Engine/Plugins/Experimental/AnimNextAnimGraph/Source/AnimNextAnimGraph/Public/TraitCore/Trait.h b/Engine/Plugins/Experimental/AnimNextAnimGraph/Source/AnimNextAnimGraph/Public/TraitCore/Trait.h
--- a/Engine/Plugins/Experimental/AnimNextAnimGraph/Source/AnimNextAnimGraph/Public/TraitCore/Trait.h
+++ b/Engine/Plugins/Experimental/AnimNextAnimGraph/Source/AnimNextAnimGraph/Public/TraitCore/Trait.h
@@ -17,6 +17,8 @@
 #include <type_traits>
 
 class FArchive;
+class URigVMController;
+struct FRigVMPinInfoArray;
 
 #define ANIM_NEXT_IMPL_DECLARE_ANIM_TRAIT_BASIC(TraitName, TraitNameHash, SuperTraitName) \
 	using TraitSuper = SuperTraitName; \
@@ -420,8 +422,11 @@ namespace UE::AnimNext
 		// instance.
 		virtual void SaveTraitSharedData(const TFunction<FString(FName PropertyName)>& GetTraitProperty, FAnimNextTraitSharedData& OutSharedData) const;
 
-		// Takes the editor properties as authored in the graph and returns the latent property metadata using UE reflection.
-		TArray<FLatentPropertyMetadata> GetLatentPropertyHandles(
+		// Takes the trait shared data and the editor properties as authored in the graph and generates the latent property metadata using UE reflection.
+		// Returns the number of latent property metadatas added to OutLatentPropertyHandles
+		virtual uint32 GetLatentPropertyHandles(
+			const FAnimNextTraitSharedData* InSharedData,
+			TArray<FLatentPropertyMetadata>& OutLatentPropertyHandles,
 			bool bFilterEditorOnly,
 			const TFunction<uint16(FName PropertyName)>& GetTraitLatentPropertyIndex) const;
 
@@ -429,6 +434,9 @@ namespace UE::AnimNext
 		virtual bool IsHidden() const { return false; }
 
 		virtual bool MultipleInstanceSupport() const { return false; }
+
+		// Enable Traits to generate TraitStack pins
+		virtual void GetProgrammaticPins(FAnimNextTraitSharedData* InSharedData, URigVMController* InController, int32 InParentPinIndex, const FString& InDefaultValue, struct FRigVMPinInfoArray& OutPinArray) const {}
 #endif
 
 	protected:
diff --git a/Engine/Plugins/Experimental/AnimNextAnimGraph/Source/AnimNextAnimGraphEditor/Private/AnimNextAnimGraphEditorModule.cpp b/Engine/Plugins/Experimental/AnimNextAnimGraph/Source/AnimNextAnimGraphEditor/Private/AnimNextAnimGraphEditorModule.cpp
--- a/Engine/Plugins/Experimental/AnimNextAnimGraph/Source/AnimNextAnimGraphEditor/Private/AnimNextAnimGraphEditorModule.cpp
+++ b/Engine/Plugins/Experimental/AnimNextAnimGraph/Source/AnimNextAnimGraphEditor/Private/AnimNextAnimGraphEditorModule.cpp
@@ -69,14 +69,20 @@ class FAnimNextAnimGraphEditorModule : public IModuleInterface
 
 		// --- AnimNextAnimationGraph ---
 		Workspace::FObjectDocumentArgs AnimNextAnimationGraphDocumentArgs(
-			Workspace::FOnRedirectWorkspaceContext::CreateLambda([](UObject* InObject)
+			Workspace::FOnRedirectWorkspaceContext::CreateLambda([](UObject* InObject) -> UObject*
 			{
+				URigVMEdGraph* EdGraph = nullptr;
+
 				UAnimNextAnimationGraph* AnimationGraph = CastChecked<UAnimNextAnimationGraph>(InObject);
 				UAnimNextAnimationGraph_EditorData* EditorData = UncookedOnly::FUtils::GetEditorData<UAnimNextAnimationGraph_EditorData>(AnimationGraph);
 
+				UAnimNextAnimationGraphEntry* AnimationGraphEntry = Cast<UAnimNextAnimationGraphEntry>(EditorData->FindEntry(FRigUnit_AnimNextGraphRoot::DefaultEntryPoint));
 				// Redirect to the inner graph
-				UAnimNextAnimationGraphEntry* AnimationGraphEntry = CastChecked<UAnimNextAnimationGraphEntry>(EditorData->FindEntry(FRigUnit_AnimNextGraphRoot::DefaultEntryPoint));
-				return AnimationGraphEntry->GetEdGraph();
+				if (ensure(AnimationGraphEntry))
+				{
+					EdGraph = AnimationGraphEntry->GetEdGraph();
+				}
+				return EdGraph;
 			}));
 
 		WorkspaceEditorModule.RegisterObjectDocumentType(FTopLevelAssetPath(TEXT("/Script/AnimNextAnimGraph.AnimNextAnimationGraph")), AnimNextAnimationGraphDocumentArgs);
diff --git a/Engine/Plugins/Experimental/AnimNextAnimGraph/Source/AnimNextAnimGraphEditor/Private/Graph/STraitEditorView.cpp b/Engine/Plugins/Experimental/AnimNextAnimGraph/Source/AnimNextAnimGraphEditor/Private/Graph/STraitEditorView.cpp
--- a/Engine/Plugins/Experimental/AnimNextAnimGraph/Source/AnimNextAnimGraphEditor/Private/Graph/STraitEditorView.cpp
+++ b/Engine/Plugins/Experimental/AnimNextAnimGraph/Source/AnimNextAnimGraphEditor/Private/Graph/STraitEditorView.cpp
@@ -368,38 +368,38 @@ void STraitEditorView::UpdateTraitStatusInStack(const TArray<TSharedPtr<FTraitDa
 		}
 	}
 
-	if (TraitData->TraitMode == ETraitMode::Additive)
-	{
-		const TArray<FTraitInterfaceUID>& RequiredInterfaces = TraitData->RequiredInterfaces;
+	const TArray<FTraitInterfaceUID>& RequiredInterfaces = TraitData->RequiredInterfaces;
 
-		const int32 NumRequired = RequiredInterfaces.Num();
-		if (NumRequired > 0)
+	const int32 NumRequired = RequiredInterfaces.Num();
+	if (NumRequired > 0)
+	{
+		for (int32 RequiredIndex = 0; RequiredIndex < NumRequired; RequiredIndex++)
 		{
-			for (int32 RequiredIndex = 0; RequiredIndex < NumRequired; RequiredIndex++)
+			const FTraitInterfaceUID& RequiredInterface = RequiredInterfaces[RequiredIndex];
+
+			bool bFound = false;
+			const int32 StartIndex = (TraitData->TraitMode == ETraitMode::Base) 
+				? CurrentTraitsData.Num() - 1 // Base traits scan all stack to find a valid interface
+				: TraitIndex; // Additive traits start search from current trait to enable traits that inherit from a trait with a required interface that is implemented in the derived class
+			
+			for (int32 SearchStartIndex = StartIndex; SearchStartIndex >= 0 && !bFound; SearchStartIndex--)
 			{
-				const FTraitInterfaceUID& RequiredInterface = RequiredInterfaces[RequiredIndex];
-
-				bool bFound = false;
-				// Start search from current trait to enable traits that inherit from a trait with a required interface that is implemented in the derived class
-				for (int32 SearchStartIndex = TraitIndex; SearchStartIndex >= 0 && !bFound; SearchStartIndex--)
-				{
-					const TSharedPtr<FTraitDataEditorDef>& ParentTraitData = CurrentTraitsData[SearchStartIndex];
+				const TSharedPtr<FTraitDataEditorDef>& ParentTraitData = CurrentTraitsData[SearchStartIndex];
 
-					const TArray<FTraitInterfaceUID>& ParentImplementedInterfaces = ParentTraitData->ImplementedInterfaces;
+				const TArray<FTraitInterfaceUID>& ParentImplementedInterfaces = ParentTraitData->ImplementedInterfaces;
 
-					for (const FTraitInterfaceUID& ParentImplementedInterface : ParentImplementedInterfaces)
+				for (const FTraitInterfaceUID& ParentImplementedInterface : ParentImplementedInterfaces)
+				{
+					if (ParentImplementedInterface == RequiredInterface)
 					{
-						if (ParentImplementedInterface == RequiredInterface)
-						{
-							bFound = true;
-							break;
-						}
+						bFound = true;
+						break;
 					}
 				}
-				if (!bFound)
-				{
-					TraitData->StackStatus.MissingInterfaces.Add(RequiredInterface);
-				}
+			}
+			if (!bFound)
+			{
+				TraitData->StackStatus.MissingInterfaces.Add(RequiredInterface);
 			}
 		}
 
diff --git a/Engine/Plugins/Experimental/AnimNextAnimGraph/Source/AnimNextAnimGraphUncookedOnly/Internal/Graph/AnimNextAnimationGraphSchema.h b/Engine/Plugins/Experimental/AnimNextAnimGraph/Source/AnimNextAnimGraphUncookedOnly/Internal/Graph/AnimNextAnimationGraphSchema.h
--- a/Engine/Plugins/Experimental/AnimNextAnimGraph/Source/AnimNextAnimGraphUncookedOnly/Internal/Graph/AnimNextAnimationGraphSchema.h
+++ b/Engine/Plugins/Experimental/AnimNextAnimGraph/Source/AnimNextAnimGraphUncookedOnly/Internal/Graph/AnimNextAnimationGraphSchema.h
@@ -16,4 +16,7 @@ public:
 
 	// Returns true if a node supports setting a node title
 	virtual bool CanSetNodeTitle(URigVMController* InController, const URigVMNode* InNode) const override;
+
+	// Returns true if a pin should be unfolded (represented by its subpins)
+	virtual bool CanUnfoldPin(URigVMController* InController, const URigVMPin* InPinToUnfold) const override;
 };
diff --git a/Engine/Plugins/Experimental/AnimNextAnimGraph/Source/AnimNextAnimGraphUncookedOnly/Private/Graph/AnimNextAnimationGraphSchema.cpp b/Engine/Plugins/Experimental/AnimNextAnimGraph/Source/AnimNextAnimGraphUncookedOnly/Private/Graph/AnimNextAnimationGraphSchema.cpp
--- a/Engine/Plugins/Experimental/AnimNextAnimGraph/Source/AnimNextAnimGraphUncookedOnly/Private/Graph/AnimNextAnimationGraphSchema.cpp
+++ b/Engine/Plugins/Experimental/AnimNextAnimGraph/Source/AnimNextAnimGraphUncookedOnly/Private/Graph/AnimNextAnimationGraphSchema.cpp
@@ -14,3 +14,22 @@ bool UAnimNextAnimationGraphSchema::CanSetNodeTitle(URigVMController* InControll
 	}
 	return false;
 }
+
+bool UAnimNextAnimationGraphSchema::CanUnfoldPin(URigVMController* InController, const URigVMPin* InPinToUnfold) const
+{
+	if (Super::CanUnfoldPin(InController, InPinToUnfold))
+	{
+		return true;
+	}
+
+	if (InPinToUnfold->IsArray())
+	{
+		// Required in order to be able to set pin default values to arrays that exist as hidden pins at trait shared data
+		// URigVMController::SetPinDefaultValue only allows setting default values to arrays if can be unfolded
+		if (InPinToUnfold->IsTraitPin() && InPinToUnfold->GetDirection() == ERigVMPinDirection::Hidden)
+		{
+			return true;
+		}
+	}
+	return false;
+}
\ No newline at end of file
diff --git a/Engine/Plugins/Experimental/AnimNextControlRig/AnimNextControlRig.uplugin b/Engine/Plugins/Experimental/AnimNextControlRig/AnimNextControlRig.uplugin
new file mode 100644
--- /dev/null
+++ b/Engine/Plugins/Experimental/AnimNextControlRig/AnimNextControlRig.uplugin
@@ -0,0 +1,46 @@
+{
+	"FileVersion" : 3,
+	"Version" : 1,
+	"VersionName" : "0.1",
+	"FriendlyName" : "Anim Next Control Rig",
+	"Description" : "Control Rig integration for Anim Next.",
+	"Category" : "Animation",
+	"CreatedBy" : "Epic Games, Inc.",
+	"CreatedByURL" : "https://epicgames.com",
+	"DocsURL" : "",
+	"MarketplaceURL" : "",
+	"SupportURL" : "",
+	"EnabledByDefault" : false,
+	"CanContainContent" : false,
+	"IsBetaVersion": false,
+	"IsExperimentalVersion": true,
+	"Installed" : false,
+	"Plugins":
+	[
+		{
+			"Name": "ControlRig",
+			"Enabled": true
+		},
+		{
+			"Name": "AnimNext",
+			"Enabled": true
+		},
+		{
+			"Name": "AnimNextAnimGraph",
+			"Enabled": true
+		},
+	],
+	"Modules" :
+	[
+		{
+			"Name" : "AnimNextControlRig",
+			"Type" : "Runtime",
+			"LoadingPhase" : "Default"
+		},
+		{
+			"Name" : "AnimNextControlRigEditor",
+			"Type" : "Editor",
+			"LoadingPhase" : "Default"
+		},		
+	]
+}
diff --git a/Engine/Plugins/Experimental/AnimNextControlRig/Source/AnimNextControlRig/AnimNextControlRig.Build.cs b/Engine/Plugins/Experimental/AnimNextControlRig/Source/AnimNextControlRig/AnimNextControlRig.Build.cs
new file mode 100644
--- /dev/null
+++ b/Engine/Plugins/Experimental/AnimNextControlRig/Source/AnimNextControlRig/AnimNextControlRig.Build.cs
@@ -0,0 +1,36 @@
+// Copyright Epic Games, Inc. All Rights Reserved.
+
+namespace UnrealBuildTool.Rules
+{
+	public class AnimNextControlRig : ModuleRules
+	{
+		public AnimNextControlRig(ReadOnlyTargetRules Target) : base(Target)
+		{
+			PublicDependencyModuleNames.AddRange(
+				new string[]
+				{
+					"CoreUObject",
+					"Engine",
+					"Core",
+					"ControlRig",
+					"ControlRigDeveloper",
+					"RigVM",
+					"AnimNext",
+					"AnimNextAnimGraph",
+				}
+			);
+
+			PrivateDependencyModuleNames.AddRange(
+				new string[]
+				{
+					"WorkspaceEditor",
+					"SlateCore",
+					"Slate",
+					"RigVMEditor",
+					"AnimNextUncookedOnly",
+					"RigVMDeveloper",
+				}
+			);
+		}
+	}
+}
diff --git a/Engine/Plugins/Experimental/AnimNextControlRig/Source/AnimNextControlRig/Internal/AnimNextControlRigModule.h b/Engine/Plugins/Experimental/AnimNextControlRig/Source/AnimNextControlRig/Internal/AnimNextControlRigModule.h
new file mode 100644
--- /dev/null
+++ b/Engine/Plugins/Experimental/AnimNextControlRig/Source/AnimNextControlRig/Internal/AnimNextControlRigModule.h
@@ -0,0 +1,26 @@
+// Copyright Epic Games, Inc. All Rights Reserved.
+#pragma once
+
+#include "CoreMinimal.h"
+#include "Modules/ModuleInterface.h"
+#include "UObject/UObjectGlobals.h"
+
+namespace UE::AnimNext::ControlRig
+{
+
+class ANIMNEXTCONTROLRIG_API FAnimNextControlRigModule : public IModuleInterface
+{
+public:
+	virtual void StartupModule() override;
+	virtual void ShutdownModule() override;
+
+#if WITH_EDITOR
+	using FReplacementObjectMap = FCoreUObjectDelegates::FReplacementObjectMap;
+	DECLARE_TS_MULTICAST_DELEGATE_OneParam(FOnObjectsReinstanced, const FReplacementObjectMap&);
+	static FOnObjectsReinstanced OnObjectsReinstanced;
+
+	FDelegateHandle OnObjectsReinstancedHandle;
+#endif
+};
+
+} // end namespace
diff --git a/Engine/Plugins/Experimental/AnimNextControlRig/Source/AnimNextControlRig/Internal/ControlRigTrait.h b/Engine/Plugins/Experimental/AnimNextControlRig/Source/AnimNextControlRig/Internal/ControlRigTrait.h
new file mode 100644
--- /dev/null
+++ b/Engine/Plugins/Experimental/AnimNextControlRig/Source/AnimNextControlRig/Internal/ControlRigTrait.h
@@ -0,0 +1,275 @@
+// Copyright Epic Games, Inc. All Rights Reserved.
+
+#pragma once
+
+#include "CoreMinimal.h"
+#include "TraitCore/Trait.h"
+#include "TraitCore/TraitSharedData.h"
+#include "TraitInterfaces/IUpdate.h"
+#include "TraitInterfaces/IEvaluate.h"
+#include "TraitInterfaces/IHierarchy.h"
+#include "TraitInterfaces/IGarbageCollection.h"
+#include "ControlRig.h"
+#include "Delegates/IDelegateInstance.h"
+#include "Tools/ControlRigHierarchyMappings.h"
+#include "Tools/ControlRigVariableMappings.h"
+#include "Tools/ControlRigIOSettings.h"
+#include "ControlRigIOMapping.h"
+// --- ---
+#include "ControlRigTrait.generated.h"
+
+class USkeletalMeshComponent;
+struct FRigVMExternalVariable;
+
+USTRUCT()
+struct FControlRigEventName
+{
+	GENERATED_BODY()
+
+	FControlRigEventName()
+	: EventName(NAME_None)
+	{}
+
+	UPROPERTY(EditAnywhere, Category = Links)
+	FName EventName;	
+};
+
+USTRUCT()
+struct FControlRigExposedProperty
+{
+	GENERATED_BODY()
+
+	FControlRigExposedProperty() = default;
+	FControlRigExposedProperty(const FName& InExposedPropertyName, bool bInIsvariable)
+		: ExposedPropertyName(InExposedPropertyName)
+		, bIsVaraible(bInIsvariable)
+	{}
+	
+	UPROPERTY(meta = (Hidden))
+	FName ExposedPropertyName = NAME_None;
+
+	UPROPERTY(meta = (Hidden))
+	bool bIsVaraible = false;
+};
+
+USTRUCT(meta = (DisplayName = "Control Rig"))
+struct FControlRigTraitSharedData : public FAnimNextTraitSharedData
+{
+	GENERATED_BODY()
+
+	/** Input to be processed. */
+	UPROPERTY()
+	FAnimNextTraitHandle Input;
+
+	UPROPERTY(EditAnywhere, Category = ControlRig, meta = (Hidden))
+	TSubclassOf<UControlRig> ControlRigClass;
+
+	// Skeleton to use as source to extract the Rig Controls. If null, the system will use the preview skeleton that was used to create the rig.
+	UPROPERTY(EditAnywhere, Category = ControlRig, meta = (Hidden))
+	TObjectPtr<USkeleton> ControlRigSkeleton;
+
+	/**
+	 * If this is checked the rig's pose needs to be reset to its initial
+	 * prior to evaluating the rig.
+	 */
+	UPROPERTY(EditAnywhere, AdvancedDisplay, Category = Settings, meta = (Hidden))
+	bool bResetInputPoseToInitial = true;
+
+	/**
+	 * If this is checked the bone pose coming from the Input will be
+	 * transferred into the Control Rig.
+	 */
+	UPROPERTY(EditAnywhere, AdvancedDisplay, Category = Settings, meta = (Hidden))
+	bool bTransferInputPose = true;
+
+	/**
+	 * If this is checked the curves coming from the AnimBP will be
+	 * transferred into the Control Rig.
+	 */
+	UPROPERTY(EditAnywhere, AdvancedDisplay, Category = Settings, meta = (Hidden))
+	bool bTransferInputCurves = true;
+
+	/**
+	* Is set, override the initial transforms with those taken from the mesh component
+	*/
+	UPROPERTY(EditAnywhere, Category=Settings, meta = (DisplayName = "Set Initial Transforms From Mesh", Hidden))
+	bool bSetRefPoseFromSkeleton = false;
+
+	/**
+	 * Transferring the pose in global space guarantees a global pose match,
+	 * while transferring in local space ensures a match of the local transforms.
+	 * In general transforms only differ if the hierarchy topology differs
+	 * between the Control Rig and the skeleton used in the AnimBP.
+	 * Note: Turning this off can potentially improve performance.
+	 */
+	UPROPERTY(EditAnywhere, AdvancedDisplay, meta = (Hidden), Category = Settings)
+	bool bTransferPoseInGlobalSpace = false;
+
+	// The customized event queue to run
+	UPROPERTY(EditAnywhere, AdvancedDisplay, meta = (Hidden), Category = Settings)
+	TArray<FControlRigEventName> EventQueue;
+
+	/**
+	 * An inclusive list of bones to transfer as part
+	 * of the input pose transfer phase.
+	 * If this list is empty all bones will be transferred.
+	 */
+	//UPROPERTY(EditAnywhere, AdvancedDisplay, meta = (Hidden), Category = Settings)
+	//TArray<FBoneReference> InputBonesToTransfer; //TODO zzz : Unsupported in AnimNext 
+
+	/**
+	 * An inclusive list of bones to transfer as part
+	 * of the output pose transfer phase.
+	 * If this list is empty all bones will be transferred.
+	 */
+	//UPROPERTY(EditAnywhere, AdvancedDisplay, meta = (Hidden), Category = Settings)
+	//TArray<FBoneReference> OutputBonesToTransfer; //TODO zzz : Unsupported in AnimNext 
+
+	// we only save mapping, 
+	// we have to query control rig when runtime 
+	// to ensure type and everything is still valid or not
+	UPROPERTY(meta = (Hidden))
+	TMap<FName, FName> InputMapping;
+
+	UPROPERTY(meta = (Hidden))
+	TMap<FName, FName> OutputMapping;
+
+	// This is the array of latent input properties
+	// This is computed at load time based on the selected ControlRig class
+	UPROPERTY(meta = (Hidden))
+	TArray<FName> LatentProperties;
+
+	UPROPERTY(meta = (Hidden))
+	TArray<FName> ExposedPropertyVariableNames;
+
+	UPROPERTY(meta = (Hidden))
+	TArray<FName> ExposedPropertyControlNames;
+	
+	UPROPERTY(meta = (Hidden))
+	TArray<ERigControlType> ExposedPropertyControlTypes;
+
+	// This is the array that maps a latent property to its memory layout in the trait instance data
+	// This is computed at load time based on the selected ControlRig class
+	UPROPERTY(Transient, meta = (Hidden))
+	TArray<uint32> LatentPropertyMemoryLayouts;
+
+	// Manual handling of latent properties
+	static void ConstructLatentProperties(const UE::AnimNext::FTraitBinding& Binding);
+
+	static void DestructLatentProperties(const UE::AnimNext::FTraitBinding& Binding);
+
+#if WITH_EDITOR
+	ANIMNEXTCONTROLRIG_API USkeleton* GetPreviewSkeleton() const;
+#endif // WITH_EDITOR
+};
+
+namespace UE::AnimNext
+{
+
+USTRUCT(meta = (DisplayName = "Control Rig"))
+struct FControlRigTrait : FBaseTrait, IEvaluate, IUpdate, IHierarchy, IGarbageCollection
+{
+	typedef FBaseTrait Super;
+
+	ANIM_NEXT_IMPL_DECLARE_ANIM_TRAIT_BASIC(FControlRigTrait, 0xf0ebdcb0, FBaseTrait)
+	ANIM_NEXT_IMPL_DECLARE_ANIM_TRAIT_INSTANCING_SUPPORT()
+	ANIM_NEXT_IMPL_DECLARE_ANIM_TRAIT_INTERFACE_SUPPORT()
+	ANIM_NEXT_IMPL_DECLARE_ANIM_TRAIT_EVENT_SUPPORT()
+
+	using FSharedData = FControlRigTraitSharedData;
+		
+	struct FInstanceData : FTrait::FInstanceData
+	{
+		FTraitPtr Input;
+
+		FRigVMDrawInterface* DebugDrawInterface = nullptr;
+		FTransform ComponentTransform;
+
+		TObjectPtr<USkeletalMeshComponent> SkeletalMeshComponent;
+
+		TObjectPtr<UControlRig> ControlRig;
+
+		FControlRigHierarchyMappings ControlRigHierarchyMappings;
+		FControlRigVariableMappings ControlRigVariableMappings;
+
+		TWeakObjectPtr<UNodeMappingContainer> NodeMappingContainer;
+		FControlRigIOSettings InputSettings;
+		FControlRigIOSettings OutputSettings;
+
+		uint16 LastBonesSerialNumberForCacheBones = 0;
+		bool bControlRigRequiresInitialization = false;
+
+		FDelegateHandle OnObjectsReinstancedHandle;
+		FDelegateHandle OnInitializedHandle;
+
+		bool bExecute = true;
+		bool bClearEventQueueRequired = false;
+		bool bUpdateInputOutputMapping = false;
+
+		void Construct(const FExecutionContext& Context, const FTraitBinding& Binding);
+		void Destruct(const FExecutionContext& Context, const FTraitBinding& Binding);
+
+		static UObject* GetAnimContext(const FExecutionContext& Context);
+		static UClass* GetTargetClass(const FSharedData* SharedData);
+
+		// Returns the property, variable name and the memory of the variable (as target) and the memory of the latent property (as source)
+		static int32 GetExposedVariablesData(const UE::AnimNext::FTraitBinding& Binding, const FSharedData* SharedData, FControlRigVariableMappings::FCustomPropertyMappings& OutMappings);
+		static int32 GetExposedControlsData(const UE::AnimNext::FTraitBinding& Binding, const FSharedData* SharedData, FControlRigVariableMappings::FCustomPropertyMappings& OutMappings);
+
+		void HandleOnInitialized_AnyThread(URigVMHost*, const FName&);
+
+		void InitializeCustomProperties(const FTraitBinding& Binding, const FSharedData* SharedData);
+
+#if WITH_EDITOR
+		void OnObjectsReinstanced(const TMap<UObject*, UObject*>& OldToNewInstanceMap);
+		bool bRegenerateVariableMappings = false;
+#endif
+	};
+
+	// FTrait impl
+	virtual void SerializeTraitSharedData(FArchive& Ar, FAnimNextTraitSharedData& SharedData) const override;
+	virtual uint32 GetNumLatentTraitProperties() const override;
+	virtual UE::AnimNext::FTraitLatentPropertyMemoryLayout GetLatentPropertyMemoryLayout(const FAnimNextTraitSharedData& SharedData, FName PropertyName, uint32 PropertyIndex) const override;
+
+	// IUpdate impl
+	virtual void OnBecomeRelevant(FUpdateTraversalContext& Context, const TTraitBinding<IUpdate>& Binding, const FTraitUpdateState& TraitState) const override;
+	virtual void PreUpdate(FUpdateTraversalContext& Context, const TTraitBinding<IUpdate>& Binding, const FTraitUpdateState& TraitState) const override;
+
+	// IEvaluate impl
+	virtual void PostEvaluate(FEvaluateTraversalContext& Context, const TTraitBinding<IEvaluate>& Binding) const override;
+
+	// IHierarchy impl
+	virtual uint32 GetNumChildren(const FExecutionContext& Context, const TTraitBinding<IHierarchy>& Binding) const override;
+	virtual void GetChildren(const FExecutionContext& Context, const TTraitBinding<IHierarchy>& Binding, FChildrenArray& Children) const override;
+
+	// IGarbageCollection impl
+	virtual void AddReferencedObjects(const FExecutionContext& Context, const TTraitBinding<IGarbageCollection>& Binding, FReferenceCollector& Collector) const override;
+
+protected:
+	static UControlRig* GetControlRig(FInstanceData* InstanceData);
+
+	static bool SetupControlRig(UObject* InAnimContext, UClass* InControlRigClass, FInstanceData* InstanceData);
+
+	static UClass* GetTargetClass(const FSharedData* SharedData, const TTraitBinding<IUpdate>& Binding);
+
+#if WITH_EDITOR
+	virtual void GetProgrammaticPins(FAnimNextTraitSharedData* InSharedData,
+		URigVMController* InController,
+		int32 InParentPinIndex,
+		const FString& InDefaultValue,
+		struct FRigVMPinInfoArray& OutPinArray) const override;
+
+	virtual uint32 GetLatentPropertyHandles(
+		const FAnimNextTraitSharedData* InSharedData,
+		TArray<FLatentPropertyMetadata>& OutLatentPropertyHandles,
+		bool bFilterEditorOnly,
+		const TFunction<uint16(FName PropertyName)>& GetTraitLatentPropertyIndex) const override;
+
+	FControlRigIOMapping::FRigControlsData RigControlsData;
+#endif
+
+	static bool GetVariableSizeAndAlignment(const FRigVMExternalVariable& Variable, uint32& PropertySize, uint32& PropertyAlignment);
+	static bool GetControlSizeAndAlignment(ERigControlType ControlType, uint32& PropertySize, uint32& PropertyAlignment);
+};
+
+} // namespace UE::AnimNext
diff --git a/Engine/Plugins/Experimental/AnimNextControlRig/Source/AnimNextControlRig/Private/AnimNextControlRigModule.cpp b/Engine/Plugins/Experimental/AnimNextControlRig/Source/AnimNextControlRig/Private/AnimNextControlRigModule.cpp
new file mode 100644
--- /dev/null
+++ b/Engine/Plugins/Experimental/AnimNextControlRig/Source/AnimNextControlRig/Private/AnimNextControlRigModule.cpp
@@ -0,0 +1,37 @@
+// Copyright Epic Games, Inc. All Rights Reserved.
+
+#include "AnimNextControlRigModule.h"
+#include "Modules/ModuleInterface.h"
+#include "Modules/ModuleManager.h"
+
+namespace UE::AnimNext::ControlRig
+{
+
+#if WITH_EDITOR
+FAnimNextControlRigModule::FOnObjectsReinstanced FAnimNextControlRigModule::OnObjectsReinstanced;
+#endif
+
+void FAnimNextControlRigModule::StartupModule()
+{
+#if WITH_EDITOR
+	// Register thread safe delegates, which allows objects to safely register to these delegates which may be called from another thread.
+	OnObjectsReinstancedHandle = FCoreUObjectDelegates::OnObjectsReinstanced.AddLambda([this](const FReplacementObjectMap& ObjectMap)
+		{
+			if (OnObjectsReinstanced.IsBound())
+			{
+				OnObjectsReinstanced.Broadcast(ObjectMap);
+			}
+		});
+#endif
+}
+
+void FAnimNextControlRigModule::ShutdownModule()
+{
+#if WITH_EDITOR
+	FCoreUObjectDelegates::OnObjectsReinstanced.Remove(OnObjectsReinstancedHandle);
+#endif
+}
+
+} // end namespace
+
+IMPLEMENT_MODULE(UE::AnimNext::ControlRig::FAnimNextControlRigModule, AnimNextControlRig)
diff --git a/Engine/Plugins/Experimental/AnimNextControlRig/Source/AnimNextControlRig/Private/ControlRigTask.cpp b/Engine/Plugins/Experimental/AnimNextControlRig/Source/AnimNextControlRig/Private/ControlRigTask.cpp
new file mode 100644
--- /dev/null
+++ b/Engine/Plugins/Experimental/AnimNextControlRig/Source/AnimNextControlRig/Private/ControlRigTask.cpp
@@ -0,0 +1,273 @@
+// Copyright Epic Games, Inc. All Rights Reserved.
+
+#include "ControlRigTask.h"
+
+#include "EvaluationVM/EvaluationVM.h"
+#include "EvaluationVM/KeyframeState.h"
+#include "TransformArrayOperations.h"
+#include "BoneContainer.h"
+#include "Components/SkeletalMeshComponent.h"
+#include "Units/Execution/RigUnit_BeginExecution.h"
+#include "Math/Transform.h"
+#include "GenerationTools.h"
+#include "Engine/SkeletalMesh.h"
+
+static TAutoConsoleVariable<int32> CVarControlRigDisableExecutionAnimNext(TEXT("ControlRig.DisableExecutionInAnimNext"), 0, TEXT("if nonzero we disable the execution of Control Rigs inside Anim Next Trrait."));
+
+#if ENABLE_ANIM_DEBUG
+TAutoConsoleVariable<int32> CVarAnimNodeControlRigDebug(TEXT("a.AnimNode.ControlRig.Debug"), 0, TEXT("Set to 1 to turn on debug drawing for AnimNode_ControlRigBase"));
+#endif
+
+DEFINE_STAT(STAT_AnimNext_Task_ControlRig);
+
+FAnimNextControlRigTask FAnimNextControlRigTask::Make(const UE::AnimNext::FControlRigTrait::FSharedData* SharedData, UE::AnimNext::FControlRigTrait::FInstanceData* InstanceData)
+{
+	FAnimNextControlRigTask Task;
+	Task.InstanceData = InstanceData;
+	Task.SharedData = SharedData;
+	return Task;
+}
+
+void FAnimNextControlRigTask::Execute(UE::AnimNext::FEvaluationVM& VM) const
+{
+	SCOPE_CYCLE_COUNTER(STAT_AnimNext_Task_ControlRig);
+	DECLARE_SCOPE_HIERARCHICAL_COUNTER_FUNC()
+
+	using namespace UE::AnimNext;
+
+	TUniquePtr<FKeyframeState> KeyframeOut;
+
+	// Try to get a Keyframe from the stack
+	if (!VM.PopValue(UE::AnimNext::KEYFRAME_STACK_NAME, KeyframeOut))
+	{
+		// if there is none, just create a reference frame
+		KeyframeOut = MakeUnique<FKeyframeState>(VM.MakeReferenceKeyframe(false));
+	}
+
+	ExecuteControlRig(VM, *KeyframeOut.Get(), InstanceData->ControlRig);
+
+	// Push our blended result back
+	VM.PushValue(KEYFRAME_STACK_NAME, MoveTemp(KeyframeOut));
+}
+
+void FAnimNextControlRigTask::ExecuteControlRig(UE::AnimNext::FEvaluationVM& VM, UE::AnimNext::FKeyframeState& KeyFrameState, UControlRig* ControlRig) const
+{
+	DECLARE_SCOPE_HIERARCHICAL_COUNTER_FUNC()
+
+	if (ControlRig)
+	{
+		FMemMark Mark(FMemStack::Get());
+
+		// Before we start modifying the RigHierarchy, we need to lock the rig to avoid corrupted state
+		// TODO : Is this really needed ? 
+		//FScopeLock LockRig(&ControlRig->GetEvaluateMutex());
+
+		URigHierarchy* Hierarchy = ControlRig->GetHierarchy();
+		if (Hierarchy == nullptr)
+		{
+			return;
+		}
+
+		// TODO : hack to get things rolling, as we don't have AnimInstanceProxy to obtain RequiredBones, until ControlRig helpers can use anim next data
+		FBoneContainer RequiredBones;
+		TArray<FBoneIndexType> LODBoneIndexToMeshBoneIndex(KeyFrameState.Pose.RefPose->GetLODBoneIndexToMeshBoneIndexMap(VM.GetCurrentLOD()));
+		LODBoneIndexToMeshBoneIndex.Sort();
+		RequiredBones.InitializeTo(LODBoneIndexToMeshBoneIndex, InstanceData->SkeletalMeshComponent->GetCurveFilterSettings(), *InstanceData->SkeletalMeshComponent->GetSkeletalMeshAsset());
+		RequiredBones.SetRefPoseOverride(InstanceData->SkeletalMeshComponent->GetRefPoseOverride());
+
+		FPoseContext Output(RequiredBones);
+
+		UE::AnimNext::FGenerationTools::RemapPose(KeyFrameState.Pose, Output);
+		Output.Curve.CopyFrom(KeyFrameState.Curves);
+		UE::AnimNext::FGenerationTools::RemapAttributes(KeyFrameState.Pose, KeyFrameState.Attributes, Output);
+
+		// temporarily give control rig access to the stack allocated attribute container
+		// control rig may have rig units that can add/get attributes to/from this container
+		UControlRig::FAnimAttributeContainerPtrScope AttributeScope(ControlRig, Output.CustomAttributes);
+
+		if (InstanceData->bUpdateInputOutputMapping)
+		{
+			// UpdateInputOutputMappingIfRequired was done in CacheBones, but there is no AnimNext equivalent
+			InstanceData->ControlRigHierarchyMappings.UpdateInputOutputMappingIfRequired(ControlRig
+				, ControlRig->GetHierarchy()
+				, Output.Pose.GetBoneContainer()
+				, TArray<FBoneReference>() //SharedData->InputBonesToTransfer // TODO zzz : Unsuported in AnimNext
+				, TArray<FBoneReference>() //SharedData->OutputBonesToTransfer // TODO zzz : Unsuported in AnimNext
+				, InstanceData->NodeMappingContainer
+				, SharedData->bTransferPoseInGlobalSpace
+				, SharedData->bResetInputPoseToInitial);
+
+			InstanceData->bUpdateInputOutputMapping = false;
+		}
+
+		if (InstanceData->ControlRigHierarchyMappings.IsPoseAdapterEnabled())
+		{
+			if (!InstanceData->ControlRigHierarchyMappings.IsUpdateToDate(Hierarchy))
+			{
+				InstanceData->ControlRigHierarchyMappings.PerformUpdateToDate(ControlRig
+					, Hierarchy
+					, Output.Pose.GetBoneContainer()
+					, InstanceData->NodeMappingContainer
+					, SharedData->bTransferPoseInGlobalSpace
+					, SharedData->bResetInputPoseToInitial);
+			}
+		}
+
+		// first update input to the system
+		UpdateInput(VM, KeyFrameState, ControlRig, Output);
+
+		if (InstanceData->bExecute)
+		{
+			const TGuardValue<bool> ResetCurrentTransfromsAfterConstructionGuard = ControlRig->GetResetCurrentTransformsAfterConstructionGuard(true);
+
+#if WITH_EDITOR
+			if (Hierarchy->IsTracingChanges())
+			{
+				Hierarchy->StorePoseForTrace(TEXT("FAnimNode_ControlRigBase::BeforeEvaluate"));
+			}
+#endif
+
+			// pick the event to run
+			if (SharedData->EventQueue.IsEmpty())
+			{
+				if (InstanceData->bClearEventQueueRequired)
+				{
+					ControlRig->SetEventQueue({ FRigUnit_BeginExecution::EventName });
+					InstanceData->bClearEventQueueRequired = false;
+				}
+			}
+			else
+			{
+				TArray<FName> EventNames;
+				Algo::Transform(SharedData->EventQueue, EventNames, [](const FControlRigEventName& InEventName)
+					{
+						return InEventName.EventName;
+					});
+				ControlRig->SetEventQueue(EventNames);
+				InstanceData->bClearEventQueueRequired = true;
+			}
+
+			if (ControlRig->IsAdditive())
+			{
+				ControlRig->ClearPoseBeforeBackwardsSolve();
+			}
+
+			// evaluate control rig
+			// UpdateGetAssetUserDataDelegate(ControlRig); // Removed in AnimNext, do we need an equivalent ?
+			ControlRig->Evaluate_AnyThread();
+
+#if UE_ENABLE_DEBUG_DRAWING
+			// When Control Rig is at editing time (in CR editor), draw instructions are consumed by ControlRigEditMode, so we need to skip drawing here.
+			const bool bShowDebug = (CVarAnimNodeControlRigDebug.GetValueOnAnyThread() == 1 && ControlRig->ExecutionType != ERigExecutionType::Editing);
+			if (bShowDebug)
+			{
+				if (FRigVMDrawInterface* DebugDrawInterface = InstanceData->DebugDrawInterface)
+				{
+					QueueControlRigDrawInstructions(ControlRig, InstanceData->DebugDrawInterface, InstanceData->ComponentTransform);
+				}
+			}
+#endif
+
+#if WITH_EDITOR
+			if (Hierarchy->IsTracingChanges())
+			{
+				Hierarchy->StorePoseForTrace(TEXT("FAnimNode_ControlRigBase::AfterEvaluate"));
+			}
+#endif
+		}
+
+		// now update output
+		UpdateOutput(VM, KeyFrameState, ControlRig, Output);
+
+		UE::AnimNext::FGenerationTools::RemapPose(Output, KeyFrameState.Pose);
+		KeyFrameState.Curves.CopyFrom(Output.Curve);
+		UE::AnimNext::FGenerationTools::RemapAttributes(Output, KeyFrameState.Pose, KeyFrameState.Attributes);
+	}
+}
+
+void FAnimNextControlRigTask::UpdateInput(UE::AnimNext::FEvaluationVM& VM, UE::AnimNext::FKeyframeState& KeyFrameState, UControlRig* ControlRig, FPoseContext& InOutput) const
+{
+	DECLARE_SCOPE_HIERARCHICAL_COUNTER_FUNC()
+
+	if (!CanExecute(ControlRig))
+	{
+		return;
+	}
+
+	DECLARE_SCOPE_HIERARCHICAL_COUNTER_FUNC()
+
+	InstanceData->ControlRigHierarchyMappings.UpdateInput(ControlRig
+		, InOutput
+		, InstanceData->InputSettings
+		, InstanceData->OutputSettings
+		, InstanceData->NodeMappingContainer
+		, InstanceData->bExecute
+		, SharedData->bTransferInputPose
+		, SharedData->bResetInputPoseToInitial
+		, SharedData->bTransferPoseInGlobalSpace
+		, SharedData->bTransferInputCurves);
+
+	InstanceData->ControlRigVariableMappings.UpdateCurveInputs(InstanceData->ControlRig, SharedData->InputMapping, InOutput.Curve);
+}
+
+void FAnimNextControlRigTask::UpdateOutput(UE::AnimNext::FEvaluationVM& VM, UE::AnimNext::FKeyframeState& KeyFrameState, UControlRig* ControlRig, FPoseContext& InOutput) const
+{
+	DECLARE_SCOPE_HIERARCHICAL_COUNTER_FUNC()
+
+	if (!CanExecute(ControlRig))
+	{
+		return;
+	}
+
+	DECLARE_SCOPE_HIERARCHICAL_COUNTER_FUNC()
+
+	InstanceData->ControlRigHierarchyMappings.UpdateOutput(ControlRig
+		, InOutput
+		, InstanceData->OutputSettings
+		, InstanceData->NodeMappingContainer
+		, InstanceData->bExecute
+		, SharedData->bTransferPoseInGlobalSpace);
+
+	InstanceData->ControlRigVariableMappings.UpdateCurveOutputs(InstanceData->ControlRig, SharedData->OutputMapping, InOutput.Curve);
+}
+
+bool FAnimNextControlRigTask::CanExecute(const UControlRig* ControlRig) const
+{
+	if (CVarControlRigDisableExecutionAnimNext->GetInt() != 0)
+	{
+		return false;
+	}
+
+	if (!InstanceData->ControlRigHierarchyMappings.CanExecute())
+	{
+		return false;
+	}
+
+	if (ControlRig)
+	{
+		return ControlRig->CanExecute();
+	}
+
+	return false;
+}
+
+void FAnimNextControlRigTask::QueueControlRigDrawInstructions(UControlRig* ControlRig, FRigVMDrawInterface* DebugDrawInterface, const FTransform& ComponentTransform) const
+{
+	ensure(ControlRig);
+	ensure(DebugDrawInterface);
+
+	if (ControlRig && DebugDrawInterface)
+	{
+		for (FRigVMDrawInstruction& Instruction : ControlRig->GetDrawInterface().Instructions)
+		{
+			if (!Instruction.IsValid())
+			{
+				continue;
+			}
+
+			const FTransform InstructionTransform = Instruction.Transform * ComponentTransform;
+			Instruction.Transform = InstructionTransform;
+			DebugDrawInterface->DrawInstruction(Instruction);
+		}
+	}
+}
diff --git a/Engine/Plugins/Experimental/AnimNextControlRig/Source/AnimNextControlRig/Private/ControlRigTask.h b/Engine/Plugins/Experimental/AnimNextControlRig/Source/AnimNextControlRig/Private/ControlRigTask.h
new file mode 100644
--- /dev/null
+++ b/Engine/Plugins/Experimental/AnimNextControlRig/Source/AnimNextControlRig/Private/ControlRigTask.h
@@ -0,0 +1,39 @@
+// Copyright Epic Games, Inc. All Rights Reserved.
+
+#pragma once
+
+#include "CoreMinimal.h"
+#include "EvaluationVM/EvaluationTask.h"
+#include "Animation/AnimCurveTypes.h"
+#include "LODPose.h"
+#include "ControlRigTrait.h"
+#include "AnimNextStats.h"
+// --- ---
+#include "ControlRigTask.generated.h"
+
+DECLARE_CYCLE_STAT_EXTERN(TEXT("AnimNext Task: ControlRig"), STAT_AnimNext_Task_ControlRig, STATGROUP_AnimNext, ANIMNEXTCONTROLRIG_API);
+
+USTRUCT()
+struct ANIMNEXTCONTROLRIG_API FAnimNextControlRigTask : public FAnimNextEvaluationTask
+{
+	GENERATED_BODY()
+
+	DECLARE_ANIM_EVALUATION_TASK(FAnimNextControlRigTask)
+
+	static FAnimNextControlRigTask Make(const UE::AnimNext::FControlRigTrait::FSharedData* SharedData, UE::AnimNext::FControlRigTrait::FInstanceData* InstanceData);
+
+	// Task entry point
+	virtual void Execute(UE::AnimNext::FEvaluationVM& VM) const override;
+
+private:
+	void ExecuteControlRig(UE::AnimNext::FEvaluationVM& VM, UE::AnimNext::FKeyframeState& KeyFrameState, UControlRig* ControlRig) const;
+	void UpdateInput(UE::AnimNext::FEvaluationVM& VM, UE::AnimNext::FKeyframeState& KeyFrameState, UControlRig* ControlRig, FPoseContext& InOutput) const;
+	void UpdateOutput(UE::AnimNext::FEvaluationVM& VM, UE::AnimNext::FKeyframeState& KeyFrameState, UControlRig* ControlRig, FPoseContext& InOutput) const;
+
+	bool CanExecute(const UControlRig* ControlRig) const;
+
+	void QueueControlRigDrawInstructions(UControlRig* ControlRig, FRigVMDrawInterface* DebugDrawInterface, const FTransform& ComponentTransform) const;
+
+	const UE::AnimNext::FControlRigTrait::FSharedData* SharedData = nullptr;
+	UE::AnimNext::FControlRigTrait::FInstanceData* InstanceData = nullptr;
+};
diff --git a/Engine/Plugins/Experimental/AnimNextControlRig/Source/AnimNextControlRig/Private/ControlRigTrait.cpp b/Engine/Plugins/Experimental/AnimNextControlRig/Source/AnimNextControlRig/Private/ControlRigTrait.cpp
new file mode 100644
--- /dev/null
+++ b/Engine/Plugins/Experimental/AnimNextControlRig/Source/AnimNextControlRig/Private/ControlRigTrait.cpp
@@ -0,0 +1,834 @@
+// Copyright Epic Games, Inc. All Rights Reserved.
+
+#include "ControlRigTrait.h"
+
+#include "Module/AnimNextModuleInstance.h"
+#include "Graph/AnimNextGraphInstance.h"
+#include "AnimNextAnimGraphSettings.h"
+#include "Animation/AnimSequence.h"
+#include "DataInterface/DataInterfaceStructAdapter.h"
+#include "TraitInterfaces/IUpdate.h"
+#include "TraitInterfaces/IEvaluate.h"
+#include "TraitInterfaces/ITimeline.h"
+#include "TraitInterfaces/IHierarchy.h"
+#include "ControlRigObjectBinding.h"
+#include "ControlRigTask.h"
+#include "AnimNextControlRigModule.h"
+#include "RigVMModel/RigVMController.h"
+#include "ControlRigBlueprint.h" // to get the preview skeleton
+#include "Components/SkeletalMeshComponent.h"
+#include "Engine/SkeletalMesh.h"
+
+//----------------------
+// --- FSharedData ---
+//----------------------
+void FControlRigTraitSharedData::ConstructLatentProperties(const UE::AnimNext::FTraitBinding& Binding)
+{
+	const FControlRigTraitSharedData* SharedData = Binding.GetSharedData<FControlRigTraitSharedData>();
+	if (SharedData && SharedData->ControlRigClass.Get() != nullptr)
+	{
+		FControlRigVariableMappings::FCustomPropertyMappings PropertyMappings;
+		UE::AnimNext::FControlRigTrait::FInstanceData::GetExposedVariablesData(Binding, SharedData, PropertyMappings);
+		UE::AnimNext::FControlRigTrait::FInstanceData::GetExposedControlsData(Binding, SharedData, PropertyMappings);
+
+		const TArray<FControlRigVariableMappings::FCustomPropertyData>& Mappings = PropertyMappings.GetMappings();
+		const int32 NumProperties = Mappings.Num();
+
+		for (const FControlRigVariableMappings::FCustomPropertyData& Mapping : Mappings)
+		{
+			const FProperty* Property = Mapping.Property;
+			const uint8* LatentPinMemory = Mapping.SourceMemory;
+			if (Property && LatentPinMemory)
+			{
+				// Init latent memory to default value
+				uint8* MutableMemory = const_cast<uint8*>(LatentPinMemory);
+				Property->InitializeValue(MutableMemory);
+			}
+		}
+	}
+}
+
+void FControlRigTraitSharedData::DestructLatentProperties(const UE::AnimNext::FTraitBinding& Binding)
+{
+	const FControlRigTraitSharedData* SharedData = Binding.GetSharedData<FControlRigTraitSharedData>();
+	if (SharedData && SharedData->ControlRigClass.Get() != nullptr)
+	{
+		FControlRigVariableMappings::FCustomPropertyMappings PropertyMappings;
+		UE::AnimNext::FControlRigTrait::FInstanceData::GetExposedVariablesData(Binding, SharedData, PropertyMappings);
+		UE::AnimNext::FControlRigTrait::FInstanceData::GetExposedControlsData(Binding, SharedData, PropertyMappings);
+
+		const TArray<FControlRigVariableMappings::FCustomPropertyData>& Mappings = PropertyMappings.GetMappings();
+		const int32 NumProperties = Mappings.Num();
+
+		for (const FControlRigVariableMappings::FCustomPropertyData& Mapping : Mappings)
+		{
+			const FProperty* Property = Mapping.Property;
+			const uint8* LatentPinMemory = Mapping.SourceMemory;
+			if (Property && LatentPinMemory)
+			{
+				uint8* MutableMemory = const_cast<uint8*>(LatentPinMemory);
+				Property->DestroyValue(MutableMemory);
+			}
+		}
+	}
+}
+
+#if WITH_EDITOR
+USkeleton* FControlRigTraitSharedData::GetPreviewSkeleton() const
+{
+	USkeleton* PreviewSkeleton = ControlRigSkeleton;
+
+	if (PreviewSkeleton == nullptr && ControlRigClass != nullptr)
+	{
+		// If the user has not provided an explicit skeleton to use, 
+		// as AnimNext does not have a preview skeleton, I get the one that was used to generate the rig
+		// (note that this might to be valid for some constructions and the user might have to provide the skeleton)
+		if (const UControlRigBlueprint* RigVMBlueprint = Cast<UControlRigBlueprint>(ControlRigClass->ClassGeneratedBy))
+		{
+			if (USkeletalMesh* SkeletalMesh = RigVMBlueprint->GetPreviewMesh())
+			{
+				PreviewSkeleton = SkeletalMesh->GetSkeleton();
+			}
+		}
+	}
+
+	return PreviewSkeleton;
+}
+#endif //WITH_EDITOR
+
+namespace UE::AnimNext
+{
+
+AUTO_REGISTER_ANIM_TRAIT(FControlRigTrait)
+
+// Trait implementation boilerplate
+#define TRAIT_INTERFACE_ENUMERATOR(GeneratorMacro) \
+	GeneratorMacro(IEvaluate) \
+	GeneratorMacro(IUpdate) \
+	GeneratorMacro(IHierarchy) \
+	GeneratorMacro(IGarbageCollection) \
+
+ANIM_NEXT_IMPL_DEFINE_ANIM_TRAIT(FControlRigTrait)
+ANIM_NEXT_IMPL_DEFINE_ANIM_TRAIT_GET_INTERFACE(FControlRigTrait, TRAIT_INTERFACE_ENUMERATOR)
+ANIM_NEXT_IMPL_DEFINE_ANIM_TRAIT_GET_INTERFACES(FControlRigTrait, TRAIT_INTERFACE_ENUMERATOR)
+ANIM_NEXT_IMPL_DEFINE_ANIM_TRAIT_GET_REQUIRED_INTERFACES(FControlRigTrait, NULL_ANIM_TRAIT_INTERFACE_ENUMERATOR)
+ANIM_NEXT_IMPL_DEFINE_ANIM_TRAIT_ON_TRAIT_EVENT(FControlRigTrait, NULL_ANIM_TRAIT_EVENT_ENUMERATOR)
+ANIM_NEXT_IMPL_DEFINE_ANIM_TRAIT_GET_TRAIT_EVENTS(FControlRigTrait, NULL_ANIM_TRAIT_EVENT_ENUMERATOR)
+#undef TRAIT_INTERFACE_ENUMERATOR
+
+//----------------------
+// --- FInstanceData ---
+//----------------------
+
+void FControlRigTrait::FInstanceData::Construct(const FExecutionContext& Context, const FTraitBinding& Binding)
+{
+	FTrait::FInstanceData::Construct(Context, Binding);
+
+	IGarbageCollection::RegisterWithGC(Context, Binding);
+
+#if WITH_EDITOR
+	{
+		OnObjectsReinstancedHandle = UE::AnimNext::ControlRig::FAnimNextControlRigModule::OnObjectsReinstanced.AddRaw(this, &FControlRigTrait::FInstanceData::OnObjectsReinstanced);
+	}
+#endif
+
+	const FSharedData* SharedData = Binding.GetSharedData<FSharedData>();
+
+	UObject* AnimContext = GetAnimContext(Context);
+	if (UClass* ControlRigClass = GetTargetClass(SharedData))
+	{
+		// we also need access to the properties when running construction event
+		//ControlRigVariableMappings.PropagateCustomInputProperties(ControlRig); // there is no control rig at this point, how is this supposed to work ?
+
+		if (ControlRigClass && ensure(SetupControlRig(AnimContext, ControlRigClass, this)))
+		{
+			OnInitializedHandle = ControlRig->OnInitialized_AnyThread().AddRaw(this, &FControlRigTrait::FInstanceData::HandleOnInitialized_AnyThread);
+		}
+
+		ControlRigVariableMappings.ResetCurvesInputToControlCache();
+		ControlRigVariableMappings.CacheCurveMappings(SharedData->InputMapping, SharedData->OutputMapping, ControlRig->GetHierarchy());
+
+		InitializeCustomProperties(Binding, SharedData);
+	}
+
+#if UE_ENABLE_DEBUG_DRAWING
+	DebugDrawInterface = Context.GetDebugDrawInterface();
+#endif
+}
+
+void FControlRigTrait::FInstanceData::Destruct(const FExecutionContext& Context, const FTraitBinding& Binding)
+{
+	FTrait::FInstanceData::Destruct(Context, Binding);
+
+	IGarbageCollection::UnregisterWithGC(Context, Binding);
+
+#if WITH_EDITOR
+	UE::AnimNext::ControlRig::FAnimNextControlRigModule::OnObjectsReinstanced.Remove(OnObjectsReinstancedHandle);
+#endif
+
+	if (ControlRig.Get() != nullptr)
+	{
+		if (OnInitializedHandle.IsValid())
+		{
+			ControlRig->OnInitialized_AnyThread().Remove(OnInitializedHandle);
+			OnInitializedHandle.Reset();
+		}
+		ControlRig = nullptr;
+	}
+}
+
+UObject* FControlRigTrait::FInstanceData::GetAnimContext(const FExecutionContext& Context)
+{
+	UObject* AnimContext = nullptr;
+
+	if (const FAnimNextModuleInstance* ModuleInstance = Context.GetRootGraphInstance().GetModuleInstance())
+	{
+		AnimContext = ModuleInstance->GetObject();
+	}
+
+	return AnimContext;
+}
+
+UClass* FControlRigTrait::FInstanceData::GetTargetClass(const FSharedData* SharedData)
+{
+	if (SharedData != nullptr)
+	{
+		return SharedData->ControlRigClass;
+	}
+	return nullptr;
+}
+
+int32 FControlRigTrait::FInstanceData::GetExposedVariablesData(const UE::AnimNext::FTraitBinding& Binding, const FSharedData* SharedData, FControlRigVariableMappings::FCustomPropertyMappings& OutMappings)
+{
+	int32 NumElementsAdded = 0;
+
+	if (SharedData && SharedData->ControlRigClass.Get() != nullptr && SharedData->ExposedPropertyVariableNames.Num() > 0)
+	{
+		if (const UControlRig* CDO = SharedData->ControlRigClass->GetDefaultObject<UControlRig>())
+		{
+			const int32 NumLatentProperties = SharedData->LatentProperties.Num();
+			if (ensure(NumLatentProperties == SharedData->LatentPropertyMemoryLayouts.Num()))
+			{
+				const TArray<FRigVMExternalVariable> PublicVariables = CDO->GetPublicVariables();
+				for (int32 LatentPropertyIndex = 0; LatentPropertyIndex < NumLatentProperties; LatentPropertyIndex++)
+				{
+					const FName& LatentPropertyName = SharedData->LatentProperties[LatentPropertyIndex];
+
+					if (!SharedData->ExposedPropertyVariableNames.Contains(LatentPropertyName))	// Only process exposed public variables
+					{
+						continue;
+					}
+
+					const FRigVMExternalVariable* Variable = PublicVariables.FindByPredicate([&LatentPropertyName](const FRigVMExternalVariable& In)
+						{
+							return In.Name == LatentPropertyName;
+						});
+
+					if (Variable != nullptr)
+					{
+						uint32 PropertyAlignment = 0;
+						uint32 PropertySize = 0;
+						UE::AnimNext::FControlRigTrait::GetVariableSizeAndAlignment(*Variable, PropertySize, PropertyAlignment);
+						check(PropertyAlignment < MAX_uint16);
+						check(PropertySize < MAX_uint16);
+
+						if (ensure(SharedData->LatentPropertyMemoryLayouts[LatentPropertyIndex] == ((PropertySize << 16) | PropertyAlignment)))
+						{
+							const UE::AnimNext::FLatentPropertyHandle* TraitLatentPropertyHandles = Binding.GetLatentPropertyHandles();
+
+							const UE::AnimNext::FLatentPropertyHandle& LatentPropertyHandle = TraitLatentPropertyHandles[LatentPropertyIndex];
+							if (LatentPropertyHandle.IsOffsetValid())
+							{
+								const uint8* LatentPropertyMemory = const_cast<uint8*>(Binding.GetLatentProperty<uint8>(LatentPropertyHandle));
+								OutMappings.AddVariable(Variable->Name, Variable->Memory, Variable->Property, LatentPropertyMemory);
+							}
+						}
+					}
+				}
+			}
+		}
+	}
+
+	return NumElementsAdded;
+}
+
+int32 FControlRigTrait::FInstanceData::GetExposedControlsData(const UE::AnimNext::FTraitBinding& Binding, const FSharedData* SharedData, FControlRigVariableMappings::FCustomPropertyMappings& OutMappings)
+{
+	int32 NumElementsAdded = 0;
+
+	if (SharedData && SharedData->ControlRigClass.Get() != nullptr)
+	{
+		const int32 NumLatentProperties = SharedData->LatentProperties.Num();
+		if (ensure(NumLatentProperties == SharedData->LatentPropertyMemoryLayouts.Num()))
+		{
+			for (int32 LatentPropertyIndex = 0; LatentPropertyIndex < NumLatentProperties; LatentPropertyIndex++)
+			{
+				const FName& LatentPropertyName = SharedData->LatentProperties[LatentPropertyIndex];
+
+				const int32 ControlIndex = SharedData->ExposedPropertyControlNames.IndexOfByKey(LatentPropertyName);
+				if (ControlIndex == INDEX_NONE)	// Only process exposed controls
+				{
+					continue;
+				}
+
+				// Note : I can not check here if the controls exist, as I would have to instantiate the rig
+				//        Using the exposed data to fill the information
+				if (ensure(SharedData->ExposedPropertyControlNames.Num() == SharedData->ExposedPropertyControlTypes.Num()))
+				{
+					const FName ControlName = SharedData->ExposedPropertyControlNames[ControlIndex];
+					const ERigControlType ControlType = SharedData->ExposedPropertyControlTypes[ControlIndex];
+
+					uint32 PropertyAlignment = 0;
+					uint32 PropertySize = 0;
+					UE::AnimNext::FControlRigTrait::GetControlSizeAndAlignment(ControlType, PropertySize, PropertyAlignment);
+					ensure(PropertyAlignment < MAX_uint16);
+					ensure(PropertySize < MAX_uint16);
+
+					if (ensure(SharedData->LatentPropertyMemoryLayouts[LatentPropertyIndex] == ((PropertySize << 16) | PropertyAlignment)))
+					{
+						const UE::AnimNext::FLatentPropertyHandle* TraitLatentPropertyHandles = Binding.GetLatentPropertyHandles();
+
+						const UE::AnimNext::FLatentPropertyHandle& LatentPropertyHandle = TraitLatentPropertyHandles[LatentPropertyIndex];
+						if (LatentPropertyHandle.IsOffsetValid())
+						{
+							const uint8* LatentPropertyMemory = const_cast<uint8*>(Binding.GetLatentProperty<uint8>(LatentPropertyHandle));
+							OutMappings.AddControl(ControlType, ControlName, nullptr, nullptr, LatentPropertyMemory);
+						}
+					}
+				}
+			}
+		}
+	}
+
+	return NumElementsAdded;
+}
+
+void FControlRigTrait::FInstanceData::HandleOnInitialized_AnyThread(URigVMHost*, const FName&)
+{
+	ControlRigHierarchyMappings.ResetRefPoseSetterHash();
+}
+
+void FControlRigTrait::FInstanceData::InitializeCustomProperties(const FTraitBinding& Binding, const FSharedData* SharedData)
+{
+	// Obtain exposed variables and controls data, in order to init the variable mappings
+	FControlRigVariableMappings::FCustomPropertyMappings PropertyMappings;
+	GetExposedVariablesData(Binding, SharedData, PropertyMappings);
+	GetExposedControlsData(Binding, SharedData, PropertyMappings);
+
+	// Setup mappings using the latent pin memory as source (we have to copy from latent pin to external variable / rig control)
+	ControlRigVariableMappings.InitializeCustomProperties(ControlRig, PropertyMappings);
+}
+
+#if WITH_EDITOR
+void FControlRigTrait::FInstanceData::OnObjectsReinstanced(const TMap<UObject*, UObject*>& OldToNewInstanceMap)
+{
+	if (ControlRig)
+	{
+		for (const TPair<UObject*, UObject*>& Pair : OldToNewInstanceMap)
+		{
+			const UObject* NewObject = Pair.Value;
+			if ((NewObject == nullptr) || !NewObject->IsA<UControlRig>())
+			{
+				continue;
+			}
+			// NewObject->GetFName() != ControlRig->GetFName()
+			ControlRigVariableMappings.ResetCustomProperties(); // avoid crashes when properties go away
+			bRegenerateVariableMappings = true;
+
+			ControlRig->OnInitialized_AnyThread().RemoveAll(this);
+			OnInitializedHandle = ControlRig->OnInitialized_AnyThread().AddRaw(this, &FControlRigTrait::FInstanceData::HandleOnInitialized_AnyThread);
+			break;
+		}
+	}
+}
+#endif
+
+//-------------------------
+// --- FControlRigTrait ---
+//-------------------------
+// --- Custom Latent Pin Support impl ---
+void FControlRigTrait::SerializeTraitSharedData(FArchive& Ar, FAnimNextTraitSharedData& SharedData) const
+{
+	FControlRigTraitSharedData& ControlRigSharedData = static_cast<FControlRigTraitSharedData&>(SharedData);
+
+	TraitSuper::SerializeTraitSharedData(Ar, SharedData);
+
+	if (Ar.IsLoading())
+	{
+		// Compute our latent property data based on our ControlRig class
+		if (UClass* ControlRigClass = ControlRigSharedData.ControlRigClass)
+		{
+			// We build the size/alignment map for each property even if their pin isn't hooked to anything
+			// since handles are reserved for every one of them
+			const int32 NumLatentProperties = ControlRigSharedData.ExposedPropertyVariableNames.Num() + ControlRigSharedData.ExposedPropertyControlNames.Num();
+			ControlRigSharedData.LatentProperties.Reserve(NumLatentProperties);
+			ControlRigSharedData.LatentPropertyMemoryLayouts.Reserve(NumLatentProperties);
+
+			if (ControlRigSharedData.ExposedPropertyVariableNames.Num() > 0)
+			{
+				if (const UControlRig* CDO = ControlRigClass->GetDefaultObject<UControlRig>())
+				{
+					const TArray<FRigVMExternalVariable> PublicVariables = CDO->GetPublicVariables();
+					if (!ensureMsgf(PublicVariables.Num() <= (int32)GetNumLatentTraitProperties(), TEXT("The ControlRig Trait only supports up to %u input variables"), GetNumLatentTraitProperties()))
+					{
+						return;
+					}
+
+					for (const FRigVMExternalVariable& Variable : PublicVariables)
+					{
+						if (!ControlRigSharedData.ExposedPropertyVariableNames.Contains(Variable.Name))	// Only process exposed public variables
+						{
+							continue;
+						}
+
+						uint32 PropertyAlignment = 0;
+						uint32 PropertySize = 0;
+						if (GetVariableSizeAndAlignment(Variable, PropertySize, PropertyAlignment))
+						{
+							check(PropertyAlignment < MAX_uint16);
+							check(PropertySize < MAX_uint16);
+
+							ControlRigSharedData.LatentProperties.Add(Variable.Name);
+							ControlRigSharedData.LatentPropertyMemoryLayouts.Add((PropertySize << 16) | PropertyAlignment);
+						}
+					}
+				}
+
+				// Here I can not get the controls list, so I just use the exposed control names and types
+				const int32 NumControls = ControlRigSharedData.ExposedPropertyControlNames.Num();
+				check(NumControls == ControlRigSharedData.ExposedPropertyControlTypes.Num());
+
+				for (int32 i = 0; i < NumControls; ++i)
+				{
+					uint32 PropertyAlignment = 0;
+					uint32 PropertySize = 0;
+					if (GetControlSizeAndAlignment(ControlRigSharedData.ExposedPropertyControlTypes[i], PropertySize, PropertyAlignment))
+					{
+						check(PropertyAlignment < MAX_uint16);
+						check(PropertySize < MAX_uint16);
+
+						ControlRigSharedData.LatentProperties.Add(ControlRigSharedData.ExposedPropertyControlNames[i]);
+						ControlRigSharedData.LatentPropertyMemoryLayouts.Add((PropertySize << 16) | PropertyAlignment);
+					}
+				}
+			}
+		}
+	}
+}
+
+uint32 FControlRigTrait::GetNumLatentTraitProperties() const
+{
+	// Number of latent trait properties must be known ahead of time to reserve space
+	// We support a maximum number of input properties, each one will need a 2-byte handle in the shared data for each trait
+	return 64;
+}
+
+FTraitLatentPropertyMemoryLayout FControlRigTrait::GetLatentPropertyMemoryLayout(const FAnimNextTraitSharedData& SharedData, FName PropertyName, uint32 PropertyIndex) const
+{
+	const FControlRigTraitSharedData& ControlRigSharedData = static_cast<const FControlRigTraitSharedData&>(SharedData);
+
+	const int32 LatentPropertyIndex = ControlRigSharedData.LatentProperties.IndexOfByKey(PropertyName);
+	if (LatentPropertyIndex == INDEX_NONE)
+	{
+		// This property isn't being tracked, ignore it
+		return FTraitLatentPropertyMemoryLayout{ 0, 1 };
+	}
+
+	const uint32 PropertyLayout = ControlRigSharedData.LatentPropertyMemoryLayouts[LatentPropertyIndex];
+	const uint32 PropertySize = PropertyLayout >> 16;
+	const uint32 PropertyAlignment = PropertyLayout & MAX_uint16;
+
+	return FTraitLatentPropertyMemoryLayout{ PropertySize, PropertyAlignment };
+}
+
+// --- IUpdate impl --- 
+void FControlRigTrait::OnBecomeRelevant(FUpdateTraversalContext& Context, const TTraitBinding<IUpdate>& Binding, const FTraitUpdateState& TraitState) const
+{
+	const FSharedData* SharedData = Binding.GetSharedData<FSharedData>();
+	check(SharedData);
+
+	FInstanceData* InstanceData = Binding.GetInstanceData<FInstanceData>();
+	check(InstanceData);
+
+	if (UControlRig* ControlRig = GetControlRig(InstanceData))
+	{
+		ControlRig->RequestInit();
+
+		InstanceData->bControlRigRequiresInitialization = true;
+		InstanceData->LastBonesSerialNumberForCacheBones = 0;
+	}
+
+	InstanceData->ControlRigHierarchyMappings.InitializeInstance();
+	InstanceData->ControlRigHierarchyMappings.ResetRefPoseSetterHash();
+
+	InstanceData->bUpdateInputOutputMapping = true;
+}
+
+void FControlRigTrait::PreUpdate(FUpdateTraversalContext& Context, const TTraitBinding<IUpdate>& Binding, const FTraitUpdateState& TraitState) const
+{
+	IUpdate::PreUpdate(Context, Binding, TraitState);
+
+	const FSharedData* SharedData = Binding.GetSharedData<FSharedData>();
+	FInstanceData* InstanceData = Binding.GetInstanceData<FInstanceData>();
+
+	if (!InstanceData->Input.IsValid())
+	{
+		InstanceData->Input = Context.AllocateNodeInstance(Binding, SharedData->Input);
+	}
+
+	if (UControlRig* ControlRig = GetControlRig(InstanceData))
+	{
+#if WITH_EDITOR
+		if (InstanceData->bRegenerateVariableMappings)
+		{
+			InstanceData->InitializeCustomProperties(Binding, SharedData);
+			InstanceData->bRegenerateVariableMappings = false;
+		}
+#endif
+
+		const float DeltaTime = TraitState.GetDeltaTime();
+		ControlRig->SetDeltaTime(DeltaTime);
+
+		InstanceData->ControlRigVariableMappings.PropagateCustomInputProperties(ControlRig);
+	}
+}
+
+// --- IEvaluate impl --- 
+void FControlRigTrait::PostEvaluate(FEvaluateTraversalContext& Context, const TTraitBinding<IEvaluate>& Binding) const
+{
+	const FSharedData* SharedData = Binding.GetSharedData<FSharedData>();
+	FInstanceData* InstanceData = Binding.GetInstanceData<FInstanceData>();
+
+	if (UControlRig* ControlRig = GetControlRig(InstanceData))
+	{
+		// Skeletal mesh required to be able to create a FBoneContainer
+		if (UObject* AnimContext = FInstanceData::GetAnimContext(Context))
+		{
+			if (InstanceData->SkeletalMeshComponent = Cast<USkeletalMeshComponent>(FControlRigObjectBinding::GetBindableObject(AnimContext)))
+			{
+#if UE_ENABLE_DEBUG_DRAWING
+				InstanceData->ComponentTransform = InstanceData->SkeletalMeshComponent->GetComponentTransform();
+			}
+#endif
+		}
+
+		Context.AppendTask(FAnimNextControlRigTask::Make(SharedData, InstanceData));
+	}
+}
+
+// --- IHierarchy impl --- 
+uint32 FControlRigTrait::GetNumChildren(const FExecutionContext& Context, const TTraitBinding<IHierarchy>& Binding) const
+{
+	return 1;
+}
+
+void FControlRigTrait::GetChildren(const FExecutionContext& Context, const TTraitBinding<IHierarchy>& Binding, FChildrenArray& Children) const
+{
+	const FInstanceData* InstanceData = Binding.GetInstanceData<FInstanceData>();
+	Children.Add(InstanceData->Input);
+}
+
+// --- IGarbageCollection impl --- 
+void FControlRigTrait::AddReferencedObjects(const FExecutionContext& Context, const TTraitBinding<IGarbageCollection>& Binding, FReferenceCollector& Collector) const
+{
+	IGarbageCollection::AddReferencedObjects(Context, Binding, Collector);
+
+	if (FInstanceData* InstanceData = Binding.GetInstanceData<FInstanceData>())
+	{
+		if (InstanceData->ControlRig.Get() != nullptr)
+		{
+			Collector.AddReferencedObject(InstanceData->ControlRig);
+		}
+	}
+}
+
+// --- Utility --- 
+
+UControlRig* FControlRigTrait::GetControlRig(FInstanceData* InstanceData)
+{
+	if (InstanceData)
+	{
+		return InstanceData->ControlRig.Get();
+	}
+
+	return nullptr;
+}
+
+bool FControlRigTrait::SetupControlRig(UObject* InAnimContext, UClass* InControlRigClass, FInstanceData* InstanceData)
+{
+	if (InstanceData->ControlRig == nullptr)
+	{
+		if (InControlRigClass != nullptr)
+		{
+			// Let's make sure the GC isn't running when we try to create a new Control Rig.
+			{
+				FGCScopeGuard GCGuard;
+				InstanceData->ControlRig = NewObject<UControlRig>(InAnimContext, InControlRigClass);
+				InstanceData->ControlRig->AtomicallyClearInternalFlags(EInternalObjectFlags::Async);
+			}
+
+			UControlRig* ControlRig = InstanceData->ControlRig.Get();
+			ControlRig->SetObjectBinding(MakeShared<FControlRigObjectBinding>());
+			if (InAnimContext != nullptr)
+			{
+				ControlRig->GetObjectBinding()->BindToObject(InAnimContext);
+			}
+			ControlRig->Initialize(true);
+			ControlRig->RequestInit();
+			
+			InstanceData->ControlRigHierarchyMappings.ResetRefPoseSetterHash();
+			
+			//UpdateInputOutputMappingIfRequired(ControlRig, InRequiredBones);
+		}
+	}
+
+	return InstanceData->ControlRig.Get() != nullptr;
+}
+
+#if WITH_EDITOR
+void FControlRigTrait::GetProgrammaticPins(FAnimNextTraitSharedData* InSharedData, URigVMController* InController, int32 InParentPinIndex, const FString& InDefaultValue, struct FRigVMPinInfoArray& OutPinArray) const
+{
+	if (FControlRigTraitSharedData* SharedData = (FControlRigTraitSharedData*)InSharedData)
+	{
+		if (SharedData->ControlRigClass != nullptr)
+		{
+			// --- Exposed Public Variables ---
+			if (SharedData->ExposedPropertyVariableNames.Num() > 0)
+			{
+				if (const UControlRig* CDO = SharedData->ControlRigClass->GetDefaultObject<UControlRig>())
+				{
+					const TArray<FRigVMExternalVariable> PublicVariables = CDO->GetPublicVariables();
+
+					for (const FRigVMExternalVariable& Variable : PublicVariables)
+					{
+						if (!SharedData->ExposedPropertyVariableNames.Contains(Variable.Name))	// Only process exposed public variables
+						{
+							continue;
+						}
+
+						OutPinArray.AddPin(const_cast<FProperty*>(Variable.Property), InController, ERigVMPinDirection::Input, InParentPinIndex, ERigVMPinDefaultValueType::AutoDetect, Variable.Memory, true);
+					}
+				}
+			}
+
+			// --- Exposed Controls ---
+			if (SharedData->ExposedPropertyControlNames.Num() > 0)
+			{
+				if (USkeleton* PreviewSkeleton = SharedData->GetPreviewSkeleton())
+				{
+					// Obtain the controls from the RigControlsData helper. This will instantiate a rig using the provided class and cache the controls until the class changes
+					const TArray<FControlRigIOMapping::FControlsInfo>& RigControls = RigControlsData.GetControls(SharedData->ControlRigClass, PreviewSkeleton);
+
+					for (const FControlRigIOMapping::FControlsInfo& Control : RigControls)
+					{
+						const FName& ControlName = Control.Name;
+						if (!SharedData->ExposedPropertyControlNames.Contains(ControlName))	// Only process exposed controls
+						{
+							continue;
+						}
+
+						const FString& ControlDefaultValue = Control.DefaultValue;
+						const TRigVMTypeIndex TypeIndex = RigVMTypeUtils::TypeIndexFromPinType(Control.PinType);
+						if (ensure(TypeIndex != INDEX_NONE))
+						{
+							const int32 Index = OutPinArray.AddPin(InController, InParentPinIndex, ControlName, ERigVMPinDirection::Input, TypeIndex, ControlDefaultValue, ERigVMPinDefaultValueType::AutoDetect, nullptr, nullptr, true);
+							OutPinArray[Index].bIsLazy = true;
+						}
+					}
+				}
+			}
+		}
+	}
+}
+
+uint32 FControlRigTrait::GetLatentPropertyHandles(
+	const FAnimNextTraitSharedData* InSharedData,
+	TArray<FLatentPropertyMetadata>& OutLatentPropertyHandles,
+	bool bFilterEditorOnly,
+	const TFunction<uint16(FName PropertyName)>& GetTraitLatentPropertyIndex) const
+{
+	// Get shared data latent properties
+	uint32 NumLatentPinsAdded = Super::GetLatentPropertyHandles(InSharedData, OutLatentPropertyHandles, bFilterEditorOnly, GetTraitLatentPropertyIndex);
+
+	// Generate Control Rig exposed pins
+	if (FControlRigTraitSharedData* SharedData = (FControlRigTraitSharedData*)InSharedData)
+	{
+		if (SharedData->ControlRigClass != nullptr)
+		{
+			// --- Iterate over public variables --- 
+			if (SharedData->ExposedPropertyVariableNames.Num() > 0)
+			{
+				if (const UControlRig* CDO = SharedData->ControlRigClass->GetDefaultObject<UControlRig>())
+				{
+					const TArray<FRigVMExternalVariable> PublicVariables = CDO->GetPublicVariables();
+
+					for (const FRigVMExternalVariable& Variable : PublicVariables)
+					{
+						if (!SharedData->ExposedPropertyVariableNames.Contains(Variable.Name))	// Only process exposed public variables
+						{
+							continue;
+						}
+
+						const FProperty* Property = Variable.Property;
+
+						FLatentPropertyMetadata Metadata;
+						Metadata.Name = Property->GetFName();
+						Metadata.RigVMIndex = GetTraitLatentPropertyIndex(Property->GetFName());
+
+						// Always false for now, we don't support freezing yet
+						Metadata.bCanFreeze = false;
+
+						OutLatentPropertyHandles.Add(Metadata);
+						NumLatentPinsAdded++;
+					}
+				}
+			}
+
+			// --- Iterate over exposed controls ---
+			if (SharedData->ExposedPropertyControlNames.Num() > 0)
+			{
+				if (USkeleton* PreviewSkeleton = SharedData->GetPreviewSkeleton())
+				{
+					// Obtain the controls from the RigControlsData helper. This will instantiate a rig using the provided class and cache the controls until the class changes
+					const TArray<FControlRigIOMapping::FControlsInfo>& RigControls = RigControlsData.GetControls(SharedData->ControlRigClass, PreviewSkeleton);
+
+					for (const FControlRigIOMapping::FControlsInfo& Control : RigControls)
+					{
+						const FName& ControlName = Control.Name;
+						if (!SharedData->ExposedPropertyControlNames.Contains(ControlName))	// Only process exposed controls
+						{
+							continue;
+						}
+
+						FLatentPropertyMetadata Metadata;
+						Metadata.Name = ControlName;
+						Metadata.RigVMIndex = GetTraitLatentPropertyIndex(ControlName);
+
+						// Always false for now, we don't support freezing yet
+						Metadata.bCanFreeze = false;
+
+						OutLatentPropertyHandles.Add(Metadata);
+						NumLatentPinsAdded++;
+					}
+				}
+			}
+		}
+	}
+
+	return NumLatentPinsAdded;
+}
+
+#endif // WITH_EDITOR
+
+bool FControlRigTrait::GetVariableSizeAndAlignment(const FRigVMExternalVariable& Variable, uint32& PropertySize, uint32& PropertyAlignment)
+{
+	bool bValidType = false;
+
+	if (Variable.TypeObject != nullptr)
+	{
+		if (const UScriptStruct* Struct = Cast<UScriptStruct>(Variable.TypeObject))
+		{
+			PropertySize = Struct->GetStructureSize();
+			PropertyAlignment = Struct->GetMinAlignment();
+			bValidType = true;
+		}
+		else if (const UEnum* Enum = Cast<UEnum>(Variable.TypeObject))
+		{
+			PropertySize = Variable.Size;
+			PropertyAlignment = Variable.Size;
+			bValidType = true;
+		}
+		else if (const UClass* Class = Cast<UClass>(Variable.TypeObject))
+		{
+			PropertySize = Class->GetStructureSize();
+			PropertyAlignment = Class->GetMinAlignment();
+			bValidType = true;
+		}
+		else
+		{
+			ensureMsgf(false, TEXT("Unsupported ControlRig public variable type: %s"), *Variable.TypeName.ToString());
+			PropertySize = Variable.Size;
+			PropertyAlignment = Variable.Size;
+		}
+	}
+	else
+	{
+		PropertySize = Variable.Size;
+		PropertyAlignment = Variable.Size;
+		bValidType = true;
+	}
+
+	return bValidType;
+}
+
+bool FControlRigTrait::GetControlSizeAndAlignment(ERigControlType ControlType, uint32& PropertySize, uint32& PropertyAlignment)
+{
+	switch (ControlType)
+	{
+		case ERigControlType::Bool:
+		{
+			PropertySize = sizeof(bool);
+			PropertyAlignment = PropertySize;
+			return true;
+		}
+		case ERigControlType::Float:
+		case ERigControlType::ScaleFloat:
+		{
+			PropertySize = sizeof(float);
+			PropertyAlignment = PropertySize;
+			return true;
+		}
+		case ERigControlType::Integer:
+		{
+			PropertySize = sizeof(int32);
+			PropertyAlignment = PropertySize;
+			return true;
+		}
+		case ERigControlType::Vector2D:
+		{
+			if (const UScriptStruct* Struct = TBaseStructure<FVector2D>::Get())
+			{
+				PropertySize = Struct->GetStructureSize();
+				PropertyAlignment = Struct->GetMinAlignment();
+				return true;
+			}
+		}
+		case ERigControlType::Position:
+		case ERigControlType::Scale:
+		{
+			PropertySize = sizeof(FVector3f);
+			PropertyAlignment = PropertySize;
+			return true;
+		}
+		case ERigControlType::Rotator:
+		{
+			if (const UScriptStruct* Struct = TBaseStructure<FRotator>::Get())
+			{
+				PropertySize = Struct->GetStructureSize();
+				PropertyAlignment = Struct->GetMinAlignment();
+				return true;
+			}
+			break;
+		}
+		case ERigControlType::Transform:
+		case ERigControlType::TransformNoScale:
+		case ERigControlType::EulerTransform:
+		{
+			if (const UScriptStruct* Struct = TBaseStructure<FTransform>::Get())
+			{
+				PropertySize = Struct->GetStructureSize();
+				PropertyAlignment = Struct->GetMinAlignment();
+				return true;
+			}
+			break;
+		}
+		default:
+		{
+			ensureMsgf(false, TEXT("Unsupported ControlRig control type."));
+			break;
+		}
+	}
+	return false;
+}
+
+} // end namespace UE::AnimNext
diff --git a/Engine/Plugins/Experimental/AnimNextControlRig/Source/AnimNextControlRigEditor/AnimNextControlRigEditor.Build.cs b/Engine/Plugins/Experimental/AnimNextControlRig/Source/AnimNextControlRigEditor/AnimNextControlRigEditor.Build.cs
new file mode 100644
--- /dev/null
+++ b/Engine/Plugins/Experimental/AnimNextControlRig/Source/AnimNextControlRigEditor/AnimNextControlRigEditor.Build.cs
@@ -0,0 +1,40 @@
+// Copyright Epic Games, Inc. All Rights Reserved.
+
+namespace UnrealBuildTool.Rules
+{
+	public class AnimNextControlRigEditor : ModuleRules
+	{
+		public AnimNextControlRigEditor(ReadOnlyTargetRules Target) : base(Target)
+		{
+			PublicDependencyModuleNames.AddRange(
+				new string[]
+				{
+					"CoreUObject",
+					"Engine",
+					"Core",
+					"ControlRig",
+					"ControlRigDeveloper",
+					"RigVM",
+					"AnimNext",
+					"AnimNextAnimGraph",
+				}
+			);
+
+			PrivateDependencyModuleNames.AddRange(
+				new string[]
+				{
+					"UnrealEd",
+					"RigVMDeveloper",
+					"WorkspaceEditor",
+					"SlateCore",
+					"Slate",
+					"RigVMEditor",
+					"AnimNextEditor",
+					"AnimNextControlRig",
+					"AnimNextUncookedOnly",
+					"AnimNextAnimGraphUncookedOnly",
+				}
+			);
+		}
+	}
+}
diff --git a/Engine/Plugins/Experimental/AnimNextControlRig/Source/AnimNextControlRigEditor/Internal/ControlRigTraitCustomization.h b/Engine/Plugins/Experimental/AnimNextControlRig/Source/AnimNextControlRigEditor/Internal/ControlRigTraitCustomization.h
new file mode 100644
--- /dev/null
+++ b/Engine/Plugins/Experimental/AnimNextControlRig/Source/AnimNextControlRigEditor/Internal/ControlRigTraitCustomization.h
@@ -0,0 +1,58 @@
+// Copyright Epic Games, Inc. All Rights Reserved.
+
+#pragma once
+
+#include "IDetailCustomization.h"
+#include "IPropertyTypeCustomization.h"
+#include "PropertyEditorModule.h"
+#include "ControlRigIOMapping.h"
+
+class URigVMPin;
+class IDetailCategoryBuilder;
+class UEdGraphNode;
+struct FRigVMVariableMappingInfo;
+struct FControlRigTraitSharedData;
+
+namespace UE::Workspace
+{
+	class IWorkspaceEditor;
+}
+
+namespace UE::AnimNext::Editor
+{
+
+class FControlRigTraitSharedDataCustomization : public IPropertyTypeCustomization
+{
+public:
+	FControlRigTraitSharedDataCustomization() = default;
+	virtual ~FControlRigTraitSharedDataCustomization();
+
+private:
+	// IPropertyTypeCustomization interface
+	virtual void CustomizeHeader(TSharedRef<IPropertyHandle> PropertyHandle, FDetailWidgetRow& HeaderRow, IPropertyTypeCustomizationUtils& CustomizationUtils) override;
+	virtual void CustomizeChildren(TSharedRef<IPropertyHandle> PropertyHandle, IDetailChildrenBuilder& ChildBuilder, IPropertyTypeCustomizationUtils& CustomizationUtils) override;
+
+	// --- IO Mapping ---
+	void OnVariableMappingChanged(const FName& PathName, const FName& Curve, bool bInput);
+	void OnPropertyExposeCheckboxChanged(ECheckBoxState NewState, FName PropertyName);
+
+	static TSharedPtr<FStructOnScope> GetControlRigSharedData(const TSharedPtr<IPropertyHandle>& StructPropertyHandle);
+	UClass* GetTargetClass() const;
+	USkeleton* GetTargetSkeleton() const;
+
+	static const FControlRigIOMapping::FControlsInfo* GetControlInfo(const TArray<FControlRigIOMapping::FControlsInfo>& Controls, const FName& ControlName);
+	static bool IsVariableProperty(FControlRigTraitSharedData* ControlRigTraitSharedData, const FName& PropertyName);
+
+	void OnObjectsReinstanced(const TMap<UObject*, UObject*>& OldToNewInstanceMap);
+
+	TSharedPtr<IPropertyHandle> StructPropertyHandle;
+	TSharedPtr<FStructOnScope> ScopedControlRigTraitSharedData;
+
+	TWeakObjectPtr<UEdGraphNode> SelectedNodeWeak = nullptr;
+	TArray<FOptionalPinFromProperty> CustomPinProperties;
+	TSharedPtr <FControlRigIOMapping> ControlRigIOMapping;
+
+	FDelegateHandle OnObjectsReinstancedHandle;
+};
+
+}
diff --git a/Engine/Plugins/Experimental/AnimNextControlRig/Source/AnimNextControlRigEditor/Private/AnimNextControlRigEditorModule.cpp b/Engine/Plugins/Experimental/AnimNextControlRig/Source/AnimNextControlRigEditor/Private/AnimNextControlRigEditorModule.cpp
new file mode 100644
--- /dev/null
+++ b/Engine/Plugins/Experimental/AnimNextControlRig/Source/AnimNextControlRigEditor/Private/AnimNextControlRigEditorModule.cpp
@@ -0,0 +1,102 @@
+// Copyright Epic Games, Inc. All Rights Reserved.
+
+#include "AnimNextControlRigEditorModule.h"
+#include "Modules/ModuleManager.h"
+#include "RigVMCore/RigVMRegistry.h"
+#include "ControlRigTraitCustomization.h"
+#include "ControlRigTrait.h"
+#include "ControlRig.h"
+#include "IAnimNextEditorModule.h"
+#include "AnimGraphUncookedOnlyUtils.h"
+#include "AnimNextEdGraphNode.h"
+#include "IWorkspaceEditorModule.h"
+#include "Subsystems/AssetEditorSubsystem.h"
+#include "UObject/Class.h"
+#include "Editor/EditorEngine.h"
+#include "K2Node.h"
+#include "Animation/Skeleton.h"
+
+extern UNREALED_API UEditorEngine* GEditor;
+
+namespace UE::AnimNext::ControlRig::Editor
+{
+
+void FAnimNextControlRigEditorModule::StartupModule()
+{
+	static TPair<UClass*, FRigVMRegistry::ERegisterObjectOperation> const AllowedObjectTypes[] =
+	{
+		{ UControlRig::StaticClass(), FRigVMRegistry::ERegisterObjectOperation::ClassAndChildren },
+		{ UClass::StaticClass(), FRigVMRegistry::ERegisterObjectOperation::Class },
+		{ USkeleton::StaticClass(), FRigVMRegistry::ERegisterObjectOperation::Class },
+	};
+
+	static UScriptStruct* const AllowedStructTypes[] =
+	{
+		FInputScaleBias::StaticStruct(),
+		FInputAlphaBoolBlend::StaticStruct(),
+	};
+
+	FRigVMRegistry& RigVMRegistry = FRigVMRegistry::Get();
+	RigVMRegistry.RegisterObjectTypes(AllowedObjectTypes);
+	RigVMRegistry.RegisterStructTypes(AllowedStructTypes);
+
+	FPropertyEditorModule& PropertyModule = FModuleManager::LoadModuleChecked<FPropertyEditorModule>("PropertyEditor");
+
+	// This generates all the Details panel ControlRig specific data and enables programmatic pin generation
+	PropertyModule.RegisterCustomPropertyTypeLayout(
+		"ControlRigTraitSharedData",
+		FOnGetPropertyTypeCustomizationInstance::CreateLambda([] { return MakeShared<UE::AnimNext::Editor::FControlRigTraitSharedDataCustomization>(); }));
+
+	UE::AnimNext::Editor::IAnimNextEditorModule& AnimNextEditorModule = FModuleManager::LoadModuleChecked<UE::AnimNext::Editor::IAnimNextEditorModule>("AnimNextEditor");
+
+	OnNodeDblClickDelegateHandle = AnimNextEditorModule.RegisterNodeDblClickHandler(UE::AnimNext::Editor::IAnimNextEditorModule::FNodeDblClickNotificationDelegate::CreateLambda([](const Workspace::FWorkspaceEditorContext& InContext, const UEdGraphNode* InNode)
+		{
+			if (const UAnimNextEdGraphNode* RigVMEdGraphNode = Cast<UAnimNextEdGraphNode>(InNode))
+			{
+				const URigVMNode* ModelNode = RigVMEdGraphNode->GetModelNode();
+
+				if (UncookedOnly::FAnimGraphUtils::IsTraitStackNode(ModelNode))
+				{
+					static const FString ControlRigTraitName = FString(TEXT("FControlRigTrait"));
+					static const FString ControlRigClassPinName = FString(TEXT("ControlRigClass"));
+
+					const TArray<URigVMPin*> TraitPins = ModelNode->GetTraitPins();
+					for (const URigVMPin* TraitPin : TraitPins)
+					{
+						if (TraitPin->GetName() == ControlRigTraitName)
+						{
+							for (const URigVMPin* SubPin : TraitPin->GetSubPins())
+							{
+								if (SubPin->GetName() == ControlRigClassPinName)
+								{
+									if (UClass* ResultClass = UClass::TryFindTypeSlow<UClass>(SubPin->GetDefaultValue()))
+									{
+										if (TSharedPtr<Workspace::IWorkspaceEditor> WorkspaceEditor = InContext.WorkspaceEditor)
+										{
+											// Control Rig is not integrated in AnimNext, so we open the whole BlueprintEditor for the asset
+											GEditor->GetEditorSubsystem<UAssetEditorSubsystem>()->OpenEditorForAsset(ResultClass->ClassGeneratedBy);
+											break;
+										}
+									}
+								}
+							}
+						}
+					}
+				}
+			}
+		}));
+}
+
+void FAnimNextControlRigEditorModule::ShutdownModule()
+{
+	if (FModuleManager::Get().IsModuleLoaded("AnimNextEditor"))
+	{
+		UE::AnimNext::Editor::IAnimNextEditorModule& AnimNextEditorModule = FModuleManager::LoadModuleChecked<UE::AnimNext::Editor::IAnimNextEditorModule>("AnimNextEditor");
+
+		AnimNextEditorModule.UnregisterNodeDblClickHandler(OnNodeDblClickDelegateHandle);
+	}
+}
+
+} // end namespace
+
+IMPLEMENT_MODULE(UE::AnimNext::ControlRig::Editor::FAnimNextControlRigEditorModule, AnimNextControlRigEditor)
diff --git a/Engine/Plugins/Experimental/AnimNextControlRig/Source/AnimNextControlRigEditor/Private/AnimNextControlRigEditorModule.h b/Engine/Plugins/Experimental/AnimNextControlRig/Source/AnimNextControlRigEditor/Private/AnimNextControlRigEditorModule.h
new file mode 100644
--- /dev/null
+++ b/Engine/Plugins/Experimental/AnimNextControlRig/Source/AnimNextControlRigEditor/Private/AnimNextControlRigEditorModule.h
@@ -0,0 +1,21 @@
+// Copyright Epic Games, Inc. All Rights Reserved.
+#pragma once
+
+#include "CoreMinimal.h"
+#include "Modules/ModuleInterface.h"
+
+namespace UE::AnimNext::ControlRig::Editor
+{
+
+class FAnimNextControlRigEditorModule : public IModuleInterface
+{
+public:
+	virtual void StartupModule() override;
+	virtual void ShutdownModule() override;
+
+private:
+	FDelegateHandle OnNodeDblClickDelegateHandle;
+
+};
+
+} // end namespace
diff --git a/Engine/Plugins/Experimental/AnimNextControlRig/Source/AnimNextControlRigEditor/Private/ControlRigTraitCustomization.cpp b/Engine/Plugins/Experimental/AnimNextControlRig/Source/AnimNextControlRigEditor/Private/ControlRigTraitCustomization.cpp
new file mode 100644
--- /dev/null
+++ b/Engine/Plugins/Experimental/AnimNextControlRig/Source/AnimNextControlRigEditor/Private/ControlRigTraitCustomization.cpp
@@ -0,0 +1,391 @@
+// Copyright Epic Games, Inc. All Rights Reserved.
+
+#include "ControlRigTraitCustomization.h"
+
+#include "DetailLayoutBuilder.h"
+#include "DetailCategoryBuilder.h"
+#include "Logging/LogScopedVerbosityOverride.h"
+#include "IWorkspaceEditor.h"
+#include "DetailWidgetRow.h"
+#include "IDetailChildrenBuilder.h"
+#include "Widgets/SRigVMVariableMappingWidget.h"
+#include "ControlRigTrait.h"
+#include "RigVMBlueprintGeneratedClass.h"
+#include "ControlRigBlueprint.h"
+#include "Engine/SkeletalMesh.h"
+#include "IPropertyUtilities.h"
+#include "AnimNextRigVMAssetEditorData.h"
+#include "AnimGraphUncookedOnlyUtils.h"
+#include "RigVMModel/RigVMController.h"
+#include "AnimNextControlRigModule.h"
+#include "ScopedTransaction.h"
+#include "K2Node.h"
+
+#define LOCTEXT_NAMESPACE "ControlRigTraitSharedDataCustomization"
+
+namespace UE::AnimNext::Editor
+{
+
+FControlRigTraitSharedDataCustomization::~FControlRigTraitSharedDataCustomization()
+{
+	if (OnObjectsReinstancedHandle.IsValid())
+	{
+		UE::AnimNext::ControlRig::FAnimNextControlRigModule::OnObjectsReinstanced.Remove(OnObjectsReinstancedHandle);
+	}
+}
+
+void FControlRigTraitSharedDataCustomization::CustomizeHeader(TSharedRef<IPropertyHandle> InPropertyHandle, FDetailWidgetRow& InHeaderRow, IPropertyTypeCustomizationUtils& InCustomizationUtils)
+{
+	if (TSharedPtr<IPropertyUtilities> PropertyUtils = InCustomizationUtils.GetPropertyUtilities())
+	{
+		const TArray< TWeakObjectPtr<UObject> >& SelectedObjects = PropertyUtils->GetSelectedObjects();
+		if (SelectedObjects.Num() != 1)
+		{
+			return;
+		}
+
+		SelectedNodeWeak = Cast<UEdGraphNode>(SelectedObjects[0]);
+	}
+
+	StructPropertyHandle = InPropertyHandle;
+	ScopedControlRigTraitSharedData = GetControlRigSharedData(StructPropertyHandle);
+
+	FControlRigTraitSharedData* ControlRigTraitSharedData = (FControlRigTraitSharedData*)ScopedControlRigTraitSharedData->GetStructMemory();
+
+	CustomPinProperties.Reset(ControlRigTraitSharedData->ExposedPropertyVariableNames.Num() + ControlRigTraitSharedData->ExposedPropertyControlNames.Num());
+	for (const FName& PropertyName : ControlRigTraitSharedData->ExposedPropertyVariableNames)
+	{
+		FOptionalPinFromProperty OptionalPin(PropertyName, true, true, FString(), FText(), false, NAME_None, false);
+		OptionalPin.bIsOverrideEnabled = false;
+		CustomPinProperties.Add(OptionalPin);
+	}
+	for (const FName& PropertyName : ControlRigTraitSharedData->ExposedPropertyControlNames)
+	{
+		FOptionalPinFromProperty OptionalPin(PropertyName, true, true, FString(), FText(), false, NAME_None, false);
+		OptionalPin.bIsOverrideEnabled = false;
+		CustomPinProperties.Add(OptionalPin);
+	}
+	ControlRigIOMapping = MakeShared<FControlRigIOMapping>(ControlRigTraitSharedData->InputMapping, ControlRigTraitSharedData->OutputMapping, CustomPinProperties);
+	ControlRigIOMapping->GetOnPinCheckStateChangedDelegate().BindSP(this, &FControlRigTraitSharedDataCustomization::OnPropertyExposeCheckboxChanged);
+	ControlRigIOMapping->GetOnVariableMappingChanged().BindSP(this, &FControlRigTraitSharedDataCustomization::OnVariableMappingChanged);
+	ControlRigIOMapping->GetOnGetTargetSkeletonDelegate().BindSP(this, &FControlRigTraitSharedDataCustomization::GetTargetSkeleton);
+	ControlRigIOMapping->GetOnGetTargetClassDelegate().BindSP(this, &FControlRigTraitSharedDataCustomization::GetTargetClass);
+
+	OnObjectsReinstancedHandle = UE::AnimNext::ControlRig::FAnimNextControlRigModule::OnObjectsReinstanced.AddRaw(this, &FControlRigTraitSharedDataCustomization::OnObjectsReinstanced);
+
+	InHeaderRow.NameContent()
+	[
+		StructPropertyHandle->CreatePropertyNameWidget()
+	];
+}
+
+void FControlRigTraitSharedDataCustomization::CustomizeChildren(TSharedRef<IPropertyHandle> InPropertyHandle, IDetailChildrenBuilder& InChildBuilder, IPropertyTypeCustomizationUtils& InCustomizationUtils)
+{
+	// --- Add struct default members, as they would be without the customization ---
+	uint32 NumChildren = 0;
+	if (InPropertyHandle->GetNumChildren(NumChildren) == FPropertyAccess::Success)
+	{
+		for (uint32 Index = 0; Index < NumChildren; ++Index)
+		{
+			TSharedRef<IPropertyHandle> ChildPropertyHandle = InPropertyHandle->GetChildHandle(Index).ToSharedRef();
+
+			InChildBuilder.AddProperty(ChildPropertyHandle);
+		}
+	}
+
+	IDetailLayoutBuilder& DetailBuilder = InChildBuilder.GetParentCategory().GetParentLayout();
+
+	if (ControlRigIOMapping.IsValid())
+	{
+		ControlRigIOMapping->CreateVariableMappingWidget(DetailBuilder);
+	}
+}
+
+void FControlRigTraitSharedDataCustomization::OnPropertyExposeCheckboxChanged(ECheckBoxState NewState, FName PropertyName)
+{
+	if (FControlRigTraitSharedData* ControlRigTraitSharedData = (FControlRigTraitSharedData*)ScopedControlRigTraitSharedData->GetStructMemory())
+	{
+		const bool bInput = ControlRigIOMapping->IsInputProperty(PropertyName);
+		ControlRigIOMapping->SetIOMapping(bInput, PropertyName, NAME_None);
+
+		if (URigVMEdGraphNode* EdGraphNode = Cast<URigVMEdGraphNode>(SelectedNodeWeak.Get()))
+		{
+			if (URigVMNode* ModelNode = EdGraphNode->GetModelNode())
+			{
+				if (URigVMGraph* Model = EdGraphNode->GetModel())
+				{
+					URigVMController* Controller = EdGraphNode->GetController();
+
+					if (URigVMEdGraph* EdGraph = Cast<URigVMEdGraph>(SelectedNodeWeak->GetGraph()))
+					{
+						static FName ExposedPropertyVariablesName = GET_MEMBER_NAME_CHECKED(FControlRigTraitSharedData, ExposedPropertyVariableNames);
+						static FName ExposedPropertyControlsName = GET_MEMBER_NAME_CHECKED(FControlRigTraitSharedData, ExposedPropertyControlNames);
+						static FName ExposedPropertyControlTypesName = GET_MEMBER_NAME_CHECKED(FControlRigTraitSharedData, ExposedPropertyControlTypes);
+						static FString ExposedVariablesSubPath(TEXT("FControlRigTrait.") + ExposedPropertyVariablesName.ToString());
+						static FString ExpòsedControlsSubPath(TEXT("FControlRigTrait.") + ExposedPropertyControlsName.ToString());
+						static FString ExpòsedControlTypesSubPath(TEXT("FControlRigTrait.") + ExposedPropertyControlTypesName.ToString());
+
+						static const FText ExposePropertyToPinText = LOCTEXT("ExposePropertyToPin", "Expose Property to Pin");
+						static const FText RemovePropertyPinText = LOCTEXT("RemovePropertyOptionalPin", "Removed Property Optional Pin");
+						const FText& TransactionText = NewState == ECheckBoxState::Checked ? ExposePropertyToPinText : RemovePropertyPinText;
+
+						FRigVMControllerCompileBracketScope CompileScope(Controller); // avoid multiple recompilations and force a details refresh after recompile
+						FScopedTransaction Transaction(TransactionText);
+						Model->Modify();
+
+						if (NewState == ECheckBoxState::Checked)
+						{
+							const bool bIsInput = ControlRigIOMapping->IsInputProperty(PropertyName);
+							// if checked, we clear mapping and unclear all children
+							ControlRigIOMapping->SetIOMapping(bIsInput, PropertyName, NAME_None);
+						}
+
+						URigVMPin* VariablesPin = ModelNode->FindPin(ExposedVariablesSubPath);
+						URigVMPin* ControlsPin = ModelNode->FindPin(ExpòsedControlsSubPath);
+						URigVMPin* ControlsTypesPin = ModelNode->FindPin(ExpòsedControlTypesSubPath);
+						if (ensure(VariablesPin && ControlsPin && ControlsTypesPin))
+						{
+							ControlRigTraitSharedData->ExposedPropertyVariableNames.Reset(CustomPinProperties.Num());
+							ControlRigTraitSharedData->ExposedPropertyControlNames.Reset(CustomPinProperties.Num());
+							ControlRigTraitSharedData->ExposedPropertyControlTypes.Reset(CustomPinProperties.Num());
+
+							for (const FOptionalPinFromProperty& OptionalPin : CustomPinProperties)
+							{
+								if (OptionalPin.bShowPin)
+								{
+									if (IsVariableProperty(ControlRigTraitSharedData, OptionalPin.PropertyName))
+									{
+										ControlRigTraitSharedData->ExposedPropertyVariableNames.Add(OptionalPin.PropertyName);
+									}
+									else
+									{
+										const TArray<FControlRigIOMapping::FControlsInfo>& Controls = ControlRigIOMapping->GetControls();
+										if (const FControlRigIOMapping::FControlsInfo* ControlInfo = GetControlInfo(Controls, OptionalPin.PropertyName))
+										{
+											ControlRigTraitSharedData->ExposedPropertyControlNames.Add(OptionalPin.PropertyName);
+											ControlRigTraitSharedData->ExposedPropertyControlTypes.Add(ControlInfo->ControlType);
+										}
+									}
+								}
+							}
+
+							static constexpr const auto SetPinArrayDefaultValue = ([](URigVMController* Controller, URigVMPin* Pin, FControlRigTraitSharedData* ControlRigTraitSharedData, const FName& PropertyName, void* ArrayData) -> bool
+								{
+									if (const FProperty* Property = ControlRigTraitSharedData->StaticStruct()->FindPropertyByName(PropertyName))
+									{
+										FString DefaultValue;
+										Property->ExportText_Direct(DefaultValue, ArrayData, nullptr, nullptr, PPF_None);
+										if (DefaultValue.IsEmpty())
+										{
+											static FString EmptyArray(TEXT("()"));
+											DefaultValue = EmptyArray;
+										}
+										return Controller->SetPinDefaultValue(Pin->GetPinPath(), DefaultValue);
+									}
+									return false;
+								});
+
+							SetPinArrayDefaultValue(Controller, VariablesPin, ControlRigTraitSharedData, ExposedPropertyVariablesName, &ControlRigTraitSharedData->ExposedPropertyVariableNames);
+							SetPinArrayDefaultValue(Controller, ControlsPin, ControlRigTraitSharedData, ExposedPropertyControlsName, &ControlRigTraitSharedData->ExposedPropertyControlNames);
+							SetPinArrayDefaultValue(Controller, ControlsTypesPin, ControlRigTraitSharedData, ExposedPropertyControlTypesName, &ControlRigTraitSharedData->ExposedPropertyControlTypes);
+						}
+
+						Controller->RepopulatePinsOnNode(ModelNode, true, false, true);
+					}
+				}
+			}
+		}
+	}
+}
+
+void FControlRigTraitSharedDataCustomization::OnVariableMappingChanged(const FName& PathName, const FName& Curve, bool bInput)
+{
+	if (URigVMEdGraphNode* EdGraphNode = Cast<URigVMEdGraphNode>(SelectedNodeWeak.Get()))
+	{
+		if (URigVMNode* ModelNode = EdGraphNode->GetModelNode())
+		{
+			if (URigVMGraph* Model = EdGraphNode->GetModel())
+			{
+				URigVMController* Controller = EdGraphNode->GetController();
+
+				FRigVMControllerCompileBracketScope CompileScope(Controller); // avoid multiple recompilations and force a details refresh after recompile
+
+				FScopedTransaction Transaction(LOCTEXT("VariableMappingChanged", "Change Variable Mapping"));
+				Model->Modify();
+
+				// @todo: this is not enough when we start breaking down struct
+				ControlRigIOMapping->SetIOMapping(bInput, PathName, Curve);
+			}
+		}
+	}
+}
+
+TSharedPtr<FStructOnScope> FControlRigTraitSharedDataCustomization::GetControlRigSharedData(const TSharedPtr<IPropertyHandle>& StructPropertyHandle)
+{
+	TSharedPtr<FStructOnScope> ControlRigSharedData;
+
+	TArray<TSharedPtr<FStructOnScope>> OutStructs;
+	StructPropertyHandle->GetOuterStructs(OutStructs);
+
+	if (OutStructs.Num() == 1)
+	{
+		ControlRigSharedData = OutStructs[0];
+	}
+
+	return ControlRigSharedData;
+}
+
+UClass* FControlRigTraitSharedDataCustomization::GetTargetClass() const
+{
+	FControlRigTraitSharedData* ControlRigTraitSharedData = (FControlRigTraitSharedData*)ScopedControlRigTraitSharedData->GetStructMemory();
+	return ControlRigTraitSharedData->ControlRigClass;
+}
+
+USkeleton* FControlRigTraitSharedDataCustomization::GetTargetSkeleton() const
+{
+	USkeleton* TargetSkeleton = nullptr;
+
+	if (const FControlRigTraitSharedData* ControlRigTraitSharedData = (FControlRigTraitSharedData*)ScopedControlRigTraitSharedData->GetStructMemory())
+	{
+		TargetSkeleton = ControlRigTraitSharedData->GetPreviewSkeleton();
+	}
+
+	return TargetSkeleton;
+}
+
+const FControlRigIOMapping::FControlsInfo* FControlRigTraitSharedDataCustomization::GetControlInfo(const TArray<FControlRigIOMapping::FControlsInfo>& Controls, const FName& ControlName)
+{
+	const FControlRigIOMapping::FControlsInfo* ControlInfo = Controls.FindByPredicate([&ControlName](const FControlRigIOMapping::FControlsInfo& In)
+		{
+			return In.Name == ControlName;
+		});
+
+	return ControlInfo;
+}
+
+bool FControlRigTraitSharedDataCustomization::IsVariableProperty(FControlRigTraitSharedData* ControlRigTraitSharedData, const FName& PropertyName )
+{
+	if (ControlRigTraitSharedData != nullptr)
+	{
+		UClass* ControlRigClass = ControlRigTraitSharedData->ControlRigClass;
+
+		if (const UControlRig* CDO = ControlRigClass->GetDefaultObject<UControlRig>())
+		{
+			const TArray<FRigVMExternalVariable> PublicVariables = CDO->GetPublicVariables();
+
+			const FRigVMExternalVariable* PublicVar = PublicVariables.FindByPredicate([&PropertyName](const FRigVMExternalVariable& In)
+				{
+					return In.Name == PropertyName;
+				});
+
+			return PublicVar != nullptr;
+		}
+	}
+
+	return false;
+}
+
+void FControlRigTraitSharedDataCustomization::OnObjectsReinstanced(const TMap<UObject*, UObject*>& OldToNewInstanceMap)
+{
+	if (ScopedControlRigTraitSharedData.IsValid())
+	{
+		for (const TPair<UObject*, UObject*>& Pair : OldToNewInstanceMap)
+		{
+			UObject* NewObject = Pair.Value;
+			if ((NewObject != nullptr) && NewObject->IsA<UControlRig>())
+			{
+				if (FControlRigTraitSharedData* ControlRigTraitSharedData = (FControlRigTraitSharedData*)ScopedControlRigTraitSharedData->GetStructMemory())
+				{
+					if (UClass* ControlRigClass = ControlRigTraitSharedData->ControlRigClass)
+					{
+						if (const UControlRig* CDO = ControlRigClass->GetDefaultObject<UControlRig>())
+						{
+							if (URigVMEdGraphNode* EdGraphNode = Cast<URigVMEdGraphNode>(SelectedNodeWeak.Get()))
+							{
+								bool bRequiresPinRefresh = false;
+
+								const TArray<FRigVMExternalVariable> PublicVariables = CDO->GetPublicVariables();
+
+								// Check if the currently exposed properties still exist
+								for (const FOptionalPinFromProperty& PinProperty : CustomPinProperties)
+								{
+									if (!PinProperty.bShowPin) // if not exposed, no need to test
+									{
+										continue;
+									}
+
+									const FName& PropertyName = PinProperty.PropertyName;
+									const FRigVMExternalVariable* PublicVar = PublicVariables.FindByPredicate([&PropertyName](const FRigVMExternalVariable& In)
+										{
+											return In.Name == PropertyName;
+										});
+							
+									if (PublicVar != nullptr) // found a public var with the same name of the exposed pin
+									{
+										const TArray<URigVMPin*>& RootPins = EdGraphNode->GetInputPins();
+										for (const URigVMPin* RootPîn : RootPins)
+										{
+											// Try to find a node pin with the same name
+											if (const URigVMPin*const* PinPtr = RootPîn->GetSubPins().FindByPredicate([&PropertyName](const URigVMPin* In)
+												{
+													return In->GetFName() == PropertyName;
+												}))
+											{
+												const URigVMPin* Pin = (*PinPtr);
+												if (Pin->GetCPPType() != PublicVar->TypeName)
+												{
+													bRequiresPinRefresh = true;
+													break;
+												}
+											}
+											else
+											{
+												// No pin found but if it is exposed, refresh
+												if (ControlRigTraitSharedData->ExposedPropertyVariableNames.Contains(PropertyName) 
+													|| ControlRigTraitSharedData->ExposedPropertyControlNames.Contains(PropertyName))
+												{
+													bRequiresPinRefresh = true;
+													break;
+												}
+												// else all ok, no pin, but it is not exposed
+											}
+										}
+									}
+									else
+									{
+										// No public var found, maybe a control ?
+										const TArray<FControlRigIOMapping::FControlsInfo>& Controls = ControlRigIOMapping->GetControls();
+										if(Controls.FindByPredicate([&PropertyName](const FControlRigIOMapping::FControlsInfo& In)
+											{
+												return In.Name == PropertyName;
+											}) == nullptr)
+										{
+											// No control found, request a rebuild
+											bRequiresPinRefresh = true;
+											break;
+										}
+									}
+								}
+
+								if (bRequiresPinRefresh)
+								{
+									if (URigVMController* Controller = EdGraphNode->GetController())
+									{
+										FRigVMControllerCompileBracketScope CompileScope(Controller); // avoid multiple recompilations and force a details refresh after recompile
+										Controller->RepopulatePinsOnNode(EdGraphNode->GetModelNode(), true, false, true); // Force pin refresh on node
+									}
+								}
+							}
+						}
+					}
+				}
+				break;
+			}
+		}
+	}
+}
+
+} // end namespace UE::AnimNext::Editor
+
+
+#undef LOCTEXT_NAMESPACE
diff --git a/Engine/Plugins/Runtime/RigVM/Source/RigVMDeveloper/Public/EdGraph/RigVMEdGraphNode.h b/Engine/Plugins/Runtime/RigVM/Source/RigVMDeveloper/Public/EdGraph/RigVMEdGraphNode.h
--- a/Engine/Plugins/Runtime/RigVM/Source/RigVMDeveloper/Public/EdGraph/RigVMEdGraphNode.h
+++ b/Engine/Plugins/Runtime/RigVM/Source/RigVMDeveloper/Public/EdGraph/RigVMEdGraphNode.h
@@ -224,6 +224,8 @@ public:
 	/** Sets the node's content from an external client, used for preview nodes without a graph */
 	void SetModelNode(URigVMNode* InModelNode);
 
+	const TArray<URigVMPin*>& GetInputPins() const { return InputPins; }
+
 protected:
 
 	FLinearColor GetNodeProfilingColor() const;
diff --git a/Engine/Source/Runtime/Engine/Classes/Animation/AnimNodeBase.h b/Engine/Source/Runtime/Engine/Classes/Animation/AnimNodeBase.h
--- a/Engine/Source/Runtime/Engine/Classes/Animation/AnimNodeBase.h
+++ b/Engine/Source/Runtime/Engine/Classes/Animation/AnimNodeBase.h
@@ -505,6 +505,16 @@ public:
 		PreviousNodeId = SourceContext.PreviousNodeId;
 	}
 
+	// This constructor allocates a new uninitialized pose, using the provided BoneContainer (when there is no AnimInstanceProxy available)
+	ENGINE_API FPoseContext(const FBoneContainer& InRequiredBones, bool bInExpectsAdditivePose = false)
+		: FAnimationBaseContext()
+		, bExpectsAdditivePose(bInExpectsAdditivePose)
+	{
+		checkSlow(InRequiredBones.IsValid());
+		Pose.SetBoneContainer(&InRequiredBones);
+		Curve.InitFrom(InRequiredBones);
+	}
+
 	UE_DEPRECATED(5.2, "This function will be made private. It should never be called externally, use the constructor instead.")
 	void Initialize(FAnimInstanceProxy* InAnimInstanceProxy) { InitializeImpl(InAnimInstanceProxy); }
 
