Control rig trait crash fix

Wrong skeletal mesh used which resulted in a crash for multi-mesh characters.

diff --git a/Engine/Plugins/Experimental/AnimNextControlRig/Source/AnimNextControlRig/Private/ControlRigTask.cpp b/Engine/Plugins/Experimental/AnimNextControlRig/Source/AnimNextControlRig/Private/ControlRigTask.cpp
--- a/Engine/Plugins/Experimental/AnimNextControlRig/Source/AnimNextControlRig/Private/ControlRigTask.cpp
+++ b/Engine/Plugins/Experimental/AnimNextControlRig/Source/AnimNextControlRig/Private/ControlRigTask.cpp
@@ -68,18 +68,22 @@ void FAnimNextControlRigTask::ExecuteControlRig(UE::AnimNext::FEvaluationVM& VM,
 			return;
 		}
 
+		const UE::AnimNext::FLODPoseStack& Pose = KeyFrameState.Pose;
+		const UE::AnimNext::FReferencePose& RefPose = Pose.GetRefPose();
+		const USkeletalMesh* SkeletalMesh = RefPose.SkeletalMesh.Get();
+
 		// TODO : hack to get things rolling, as we don't have AnimInstanceProxy to obtain RequiredBones, until ControlRig helpers can use anim next data
 		FBoneContainer RequiredBones;
-		TArray<FBoneIndexType> LODBoneIndexToMeshBoneIndex(KeyFrameState.Pose.RefPose->GetLODBoneIndexToMeshBoneIndexMap(VM.GetCurrentLOD()));
+		TArray<FBoneIndexType> LODBoneIndexToMeshBoneIndex(RefPose.GetLODBoneIndexToMeshBoneIndexMap(VM.GetCurrentLOD()));
 		LODBoneIndexToMeshBoneIndex.Sort();
-		RequiredBones.InitializeTo(LODBoneIndexToMeshBoneIndex, InstanceData->SkeletalMeshComponent->GetCurveFilterSettings(), *InstanceData->SkeletalMeshComponent->GetSkeletalMeshAsset());
+		RequiredBones.InitializeTo(LODBoneIndexToMeshBoneIndex, InstanceData->SkeletalMeshComponent->GetCurveFilterSettings(), *SkeletalMesh);
 		RequiredBones.SetRefPoseOverride(InstanceData->SkeletalMeshComponent->GetRefPoseOverride());
 
 		FPoseContext Output(RequiredBones);
 
-		UE::AnimNext::FGenerationTools::RemapPose(KeyFrameState.Pose, Output);
+		UE::AnimNext::FGenerationTools::RemapPose(Pose, Output);
 		Output.Curve.CopyFrom(KeyFrameState.Curves);
-		UE::AnimNext::FGenerationTools::RemapAttributes(KeyFrameState.Pose, KeyFrameState.Attributes, Output);
+		UE::AnimNext::FGenerationTools::RemapAttributes(Pose, KeyFrameState.Attributes, Output);
 
 		// temporarily give control rig access to the stack allocated attribute container
 		// control rig may have rig units that can add/get attributes to/from this container
diff --git a/Engine/Source/Runtime/Engine/Private/BoneContainer.cpp b/Engine/Source/Runtime/Engine/Private/BoneContainer.cpp
--- a/Engine/Source/Runtime/Engine/Private/BoneContainer.cpp
+++ b/Engine/Source/Runtime/Engine/Private/BoneContainer.cpp
@@ -75,17 +75,17 @@ FBoneContainer::FBoneContainer(const TArrayView<const FBoneIndexType>& InRequire
 	Initialize(InCurveFilterSettings);
 }
 
-void FBoneContainer::InitializeTo(const TArrayView<const FBoneIndexType>& InRequiredBoneIndexArray, const FCurveEvaluationOption& CurveEvalOption, UObject& InAsset)
+void FBoneContainer::InitializeTo(const TArrayView<const FBoneIndexType>& InRequiredBoneIndexArray, const FCurveEvaluationOption& CurveEvalOption, const UObject& InAsset)
 {
 	const UE::Anim::FCurveFilterSettings CurveFilterSettings(CurveEvalOption.bAllowCurveEvaluation ? UE::Anim::ECurveFilterMode::DisallowFiltered : UE::Anim::ECurveFilterMode::DisallowAll, CurveEvalOption.DisallowedList, CurveEvalOption.LODIndex);
 	InitializeTo(InRequiredBoneIndexArray, CurveFilterSettings, InAsset);
 }
 
-void FBoneContainer::InitializeTo(const TArrayView<const FBoneIndexType>& InRequiredBoneIndexArray, const UE::Anim::FCurveFilterSettings& CurveFilterSettings, UObject& InAsset)
+void FBoneContainer::InitializeTo(const TArrayView<const FBoneIndexType>& InRequiredBoneIndexArray, const UE::Anim::FCurveFilterSettings& CurveFilterSettings, const UObject& InAsset)
 {
 	LLM_SCOPE_BYNAME(TEXT("Animation/BoneContainer"));
 	BoneIndicesArray = InRequiredBoneIndexArray;
-	Asset = &InAsset;
+	Asset = const_cast<UObject*>(&InAsset);
 
 	Initialize(CurveFilterSettings);
 }
diff --git a/Engine/Source/Runtime/Engine/Public/BoneContainer.h b/Engine/Source/Runtime/Engine/Public/BoneContainer.h
--- a/Engine/Source/Runtime/Engine/Public/BoneContainer.h
+++ b/Engine/Source/Runtime/Engine/Public/BoneContainer.h
@@ -255,13 +255,13 @@ public:
 
 	ENGINE_API FBoneContainer(const TArrayView<const FBoneIndexType>& InRequiredBoneIndexArray, const UE::Anim::FCurveFilterSettings& InCurveFilterSettings, UObject& InAsset);
 	
-	ENGINE_API void InitializeTo(const TArrayView<const FBoneIndexType>& InRequiredBoneIndexArray, const UE::Anim::FCurveFilterSettings& InCurveFilterSettings, UObject& InAsset);
+	ENGINE_API void InitializeTo(const TArrayView<const FBoneIndexType>& InRequiredBoneIndexArray, const UE::Anim::FCurveFilterSettings& InCurveFilterSettings, const UObject& InAsset);
 	
 	UE_DEPRECATED(5.3, "Please use the constructor that takes a FCurveFilterSettings.")
 	ENGINE_API FBoneContainer(const TArrayView<const FBoneIndexType>& InRequiredBoneIndexArray, const FCurveEvaluationOption& CurveEvalOption, UObject& InAsset);
 
 	UE_DEPRECATED(5.3, "Please use InitializeTo that takes a FCurveFilterSettings.")
-	ENGINE_API void InitializeTo(const TArrayView<const FBoneIndexType>& InRequiredBoneIndexArray, const FCurveEvaluationOption& CurveEvalOption, UObject& InAsset);
+	ENGINE_API void InitializeTo(const TArrayView<const FBoneIndexType>& InRequiredBoneIndexArray, const FCurveEvaluationOption& CurveEvalOption, const UObject& InAsset);
 
 	/** Resets the container and reclaims all allocated memory but preserve the serial number. */
 	ENGINE_API void Reset();
