AnimNext ControlRig native pose adapter and hierarchy mappings

diff --git a/Engine/Plugins/Experimental/AnimNext/Source/AnimNext/Public/LODPose.h b/Engine/Plugins/Experimental/AnimNext/Source/AnimNext/Public/LODPose.h
--- a/Engine/Plugins/Experimental/AnimNext/Source/AnimNext/Public/LODPose.h
+++ b/Engine/Plugins/Experimental/AnimNext/Source/AnimNext/Public/LODPose.h
@@ -65,6 +65,43 @@ struct FLODPose
 		}
 	}
 
+	// This function copies transforms from a AoS source
+	// Warning : The source number of transforms has to match the local pose size
+	template<typename OtherAllocator>
+	void CopyTransformsFrom(const TArray<FTransform, OtherAllocator>& SourceTransforms)
+	{
+		const int32 NumTransforms = SourceTransforms.Num();
+		check(LocalTransformsView.Num() == NumTransforms);
+
+		const FTransform* RESTRICT SourceTransform = SourceTransforms.GetData();
+		const FTransform* RESTRICT SourceTransformEnd = SourceTransforms.GetData() + NumTransforms;
+		FVector* RESTRICT TargetTranslations = &LocalTransformsView[0].Translation;
+		FQuat* RESTRICT TargetRotations = &LocalTransformsView[0].Rotation;
+		FVector* RESTRICT TargetScales = &LocalTransformsView[0].Scale3D;
+
+		for (; SourceTransform < SourceTransformEnd; SourceTransform++, TargetTranslations++, TargetRotations++, TargetScales++)
+		{
+			//LocalTransformsView[Index] = SourceTransforms[Index]; // If LocalTransforms is in SoA format, there is a conversion from AoS to SoA
+			(*TargetTranslations) = SourceTransform->GetTranslation();
+			(*TargetRotations) = SourceTransform->GetRotation();
+			(*TargetScales) = SourceTransform->GetScale3D();
+		}
+	}
+
+	// This function copies transforms to an AoS target array
+	// Warning : The target number of transforms has to match the local pose size
+	template<typename OtherAllocator>
+	void CopyTransformsTo(TArray<FTransform, OtherAllocator>& OutTransforms) const
+	{
+		const int32 NumTransforms = LocalTransformsView.Num();
+		
+		check(NumTransforms == OutTransforms.Num());
+		for (int32 Index = 0; Index < NumTransforms; Index++)
+		{
+			OutTransforms[Index] = LocalTransformsView[Index]; // If LocalTransforms is in SoA format, there is a conversion from SoA to AoS
+		}
+	}
+
 	void SetRefPose(bool bAdditive = false, int32 StartIndex = 0, int32 NumTransformsToCopy = -1)
 	{
 		const int32 NumTransforms = LocalTransformsView.Num();
diff --git a/Engine/Plugins/Experimental/AnimNext/Source/AnimNext/Public/ReferencePose.h b/Engine/Plugins/Experimental/AnimNext/Source/AnimNext/Public/ReferencePose.h
--- a/Engine/Plugins/Experimental/AnimNext/Source/AnimNext/Public/ReferencePose.h
+++ b/Engine/Plugins/Experimental/AnimNext/Source/AnimNext/Public/ReferencePose.h
@@ -207,6 +207,7 @@ struct TReferencePose
 		return MeshBoneIndexToParentMeshBoneIndexMap;
 	}
 
+	// Returns the corresponding Skeleron Bone Index to a LOD Bone Index
 	int32 GetSkeletonBoneIndexFromLODBoneIndex(int32 LODBoneIndex) const
 	{
 		const int32 NumBonesLOD0 = LODBoneIndexToSkeletonBoneIndexMapPerLOD[0].Num();
@@ -227,6 +228,29 @@ struct TReferencePose
 		return SkeletonBoneIndexToLODBoneIndexMap[SkeletionBoneIndex];
 	}
 
+	int32 GetLODParentBoneIndex(int32 LODLevel, int32 LODBoneIndex) const
+	{
+		const TArrayView<const FBoneIndexType> ParentLODMap = GetLODBoneIndexToParentLODBoneIndexMap(LODLevel);
+		if (ParentLODMap.IsValidIndex(LODBoneIndex))
+		{
+			const FBoneIndexType LODParentBoneIndex = ParentLODMap[LODBoneIndex];
+			if (LODParentBoneIndex != UINT16_MAX)
+			{
+				return LODParentBoneIndex;
+			}
+		}
+		return INDEX_NONE;
+	}
+
+	const USkeletalMesh* GetSkeletalMeshAsset() const
+	{
+		return SkeletalMesh.Get();
+	}
+	const USkeleton* GetSkeletonAsset() const
+	{
+		return Skeleton.Get();
+	}
+
 	FTransform GetRefPoseTransform(int32 LODBoneIndex) const
 	{
 		const int32 NumBonesLOD0 = LODBoneIndexToMeshBoneIndexMapPerLOD[0].Num();
diff --git a/Engine/Plugins/Experimental/AnimNextControlRig/Source/AnimNextControlRig/Internal/AnimNextControlRigHierarchyMappings.h b/Engine/Plugins/Experimental/AnimNextControlRig/Source/AnimNextControlRig/Internal/AnimNextControlRigHierarchyMappings.h
new file mode 100644
--- /dev/null
+++ b/Engine/Plugins/Experimental/AnimNextControlRig/Source/AnimNextControlRig/Internal/AnimNextControlRigHierarchyMappings.h
@@ -0,0 +1,136 @@
+// Copyright Epic Games, Inc. All Rights Reserved.
+
+#pragma once
+
+#include "CoreMinimal.h"
+#include "ReferencePose.h"
+
+class UControlRig;
+class URigHierarchy;
+class UNodeMappingContainer;
+class USkeletalMeshComponent;
+struct FBoneContainer;
+struct FBoneReference;
+struct FPoseContext;
+struct FControlRigIOSettings;
+struct FCompactPose;
+
+namespace UE::AnimNext
+{
+struct FKeyframeState;
+}
+
+namespace UE::AnimNext::ControlRig
+{
+
+class FAnimNextControlRigPoseAdapter;
+
+struct ANIMNEXTCONTROLRIG_API FAnimNextControlRigHierarchyMappings
+{
+	FAnimNextControlRigHierarchyMappings() = default;
+
+	void InitializeInstance();
+
+	void LinkToHierarchy(URigHierarchy* InHierarchy);
+
+	bool CanExecute() const
+	{
+		return PoseAdapter.IsValid();
+	}
+
+	void ResetRefPoseSetterHash()
+	{
+		RefPoseSetterHash.Reset();
+	}
+
+	void UpdateControlRigRefPoseIfNeeded(UControlRig* ControlRig
+		, const UObject* InstanceObject
+		, const USkeletalMeshComponent* SkeletalMeshComponent
+		, const UE::AnimNext::FReferencePose& InRefPose
+		, bool bInSetRefPoseFromSkeleton
+		, bool bIncludePoseInHash);
+
+	void UpdateInputOutputMappingIfRequired(UControlRig* InControlRig
+		, URigHierarchy* InHierarchy
+		, const UE::AnimNext::FReferencePose& InRefPose
+		, int32 InCurrentLOD
+		, const TArray<FBoneReference>& InInputBonesToTransfer
+		, const TArray<FBoneReference>& InOutputBonesToTransfer
+		, const TWeakObjectPtr<UNodeMappingContainer>& InNodeMappingContainer
+		, bool bInTransferPoseInGlobalSpace
+		, bool bResetInputPoseToInitial);
+
+	void UpdateInput(UControlRig* ControlRig
+		, UE::AnimNext::FKeyframeState& InOutput
+		, const FControlRigIOSettings& InInputSettings
+		, const FControlRigIOSettings& InOutputSettings
+		, TWeakObjectPtr<UNodeMappingContainer>& InNodeMappingContainer
+		, bool bInExecute
+		, bool bInTransferInputPose
+		, bool bInResetInputPoseToInitial
+		, bool bInTransferPoseInGlobalSpace
+		, bool bInTransferInputCurves);
+
+	void UpdateOutput(UControlRig* ControlRig
+		, UE::AnimNext::FKeyframeState& InOutput
+		, const FControlRigIOSettings& InOutputSettings
+		, TWeakObjectPtr<UNodeMappingContainer>& InNodeMappingContainer
+		, bool bInExecute
+		, bool bInTransferPoseInGlobalSpace);
+
+	const TArray<TPair<uint16, uint16>>& GetControlRigBoneInputMappingByIndex() const
+	{
+		return ControlRigBoneInputMappingByIndex;
+	}
+
+	TArray<TPair<uint16, uint16>>& GetControlRigBoneOutputMappingByIndex()
+	{
+		return ControlRigBoneOutputMappingByIndex;
+	}
+
+	const TMap<FName, uint16>& GetControlRigBoneInputMappingByName() const
+	{
+		return ControlRigBoneInputMappingByName;
+	}
+
+	TMap<FName, uint16>& GetControlRigBoneOutputMappingByName()
+	{
+		return ControlRigBoneOutputMappingByName;
+	}
+
+	bool CheckPoseAdapter() const;
+
+	bool IsUpdateToDate(const URigHierarchy* InHierarchy) const;
+
+	void PerformUpdateToDate(UControlRig* ControlRig
+		, URigHierarchy* InHierarchy
+		, const UE::AnimNext::FReferencePose& InRefPose
+		, int32 InCurrentLOD
+		, const TWeakObjectPtr<UNodeMappingContainer>& InNodeMappingContainer
+		, bool bInTransferPoseInGlobalSpace
+		, bool bInResetInputPoseToInitial);
+
+protected:
+	/** Complete mapping from skeleton to control rig bone index */
+	TArray<TPair<uint16, uint16>> ControlRigBoneInputMappingByIndex;
+	TArray<TPair<uint16, uint16>> ControlRigBoneOutputMappingByIndex;
+
+	/** Complete mapping from skeleton to curve name */
+	TArray<TPair<uint16, FName>> ControlRigCurveMappingByIndex;
+
+	/** Rig Hierarchy bone name to required array index mapping */
+	TMap<FName, uint16> ControlRigBoneInputMappingByName;
+	TMap<FName, uint16> ControlRigBoneOutputMappingByName;
+
+	/** Rig Curve name to Curve mapping */
+	TMap<FName, FName> ControlRigCurveMappingByName;
+
+	TArray<bool> HierarchyCurveCopied;
+
+	TSharedPtr<FAnimNextControlRigPoseAdapter> PoseAdapter;
+
+	// A hash to encode the pointer to anim instance
+	TOptional<int32> RefPoseSetterHash;
+};
+
+} // end namespace UE::AnimNext
diff --git a/Engine/Plugins/Experimental/AnimNextControlRig/Source/AnimNextControlRig/Internal/AnimNextControlRigPoseAdapter.h b/Engine/Plugins/Experimental/AnimNextControlRig/Source/AnimNextControlRig/Internal/AnimNextControlRigPoseAdapter.h
new file mode 100644
--- /dev/null
+++ b/Engine/Plugins/Experimental/AnimNextControlRig/Source/AnimNextControlRig/Internal/AnimNextControlRigPoseAdapter.h
@@ -0,0 +1,44 @@
+// Copyright Epic Games, Inc. All Rights Reserved.
+
+#pragma once
+
+#include "CoreMinimal.h"
+#include "Tools/ControlRigPoseAdapter.h"
+#include "ReferencePose.h"
+#include "LODPose.h"
+
+struct FReferenceSkeleton;
+class UNodeMappingContainer;
+class UControlRig;
+class URigHierarchy;
+
+namespace UE::AnimNext
+{
+struct FKeyframeState;
+}
+
+namespace UE::AnimNext::ControlRig
+{
+
+class ANIMNEXTCONTROLRIG_API FAnimNextControlRigPoseAdapter : public FControlRigPoseAdapter
+{
+public:
+	FAnimNextControlRigPoseAdapter() = default;
+
+	void CopyBonesFrom(const FLODPoseStack& InPose)
+	{
+		InPose.CopyTransformsTo(LocalPose);
+	}
+
+	void UpdateInputOutputMappingIfRequired(UControlRig* InControlRig
+		, URigHierarchy* InHierarchy
+		, const UE::AnimNext::FReferencePose& InRefPose
+		, const int32 InCurentLOD
+		, const TWeakObjectPtr<UNodeMappingContainer>& InNodeMappingContainer
+		, bool bInTransferPoseInGlobalSpace
+		, bool bInResetInputPoseToInitial);
+
+	static const FReferenceSkeleton* GetReferenceSkeleton(const UE::AnimNext::FReferencePose& InRefPose);
+};
+
+} // end namespace UE::AnimNext
diff --git a/Engine/Plugins/Experimental/AnimNextControlRig/Source/AnimNextControlRig/Internal/ControlRigTrait.h b/Engine/Plugins/Experimental/AnimNextControlRig/Source/AnimNextControlRig/Internal/ControlRigTrait.h
--- a/Engine/Plugins/Experimental/AnimNextControlRig/Source/AnimNextControlRig/Internal/ControlRigTrait.h
+++ b/Engine/Plugins/Experimental/AnimNextControlRig/Source/AnimNextControlRig/Internal/ControlRigTrait.h
@@ -11,7 +11,7 @@
 #include "TraitInterfaces/IGarbageCollection.h"
 #include "ControlRig.h"
 #include "Delegates/IDelegateInstance.h"
-#include "Tools/ControlRigHierarchyMappings.h"
+#include "AnimNextControlRigHierarchyMappings.h"
 #include "Tools/ControlRigVariableMappings.h"
 #include "Tools/ControlRigIOSettings.h"
 #if WITH_EDITOR
@@ -191,7 +191,7 @@ struct FControlRigTrait : FBaseTrait, IEvaluate, IUpdate, IHierarchy, IGarbageCo
 
 		TObjectPtr<UControlRig> ControlRig;
 
-		FControlRigHierarchyMappings ControlRigHierarchyMappings;
+		ControlRig::FAnimNextControlRigHierarchyMappings ControlRigHierarchyMappings;
 		FControlRigVariableMappings ControlRigVariableMappings;
 
 		TWeakObjectPtr<UNodeMappingContainer> NodeMappingContainer;
diff --git a/Engine/Plugins/Experimental/AnimNextControlRig/Source/AnimNextControlRig/Private/AnimNextControlRigHierarchyMappings.cpp b/Engine/Plugins/Experimental/AnimNextControlRig/Source/AnimNextControlRig/Private/AnimNextControlRigHierarchyMappings.cpp
new file mode 100644
--- /dev/null
+++ b/Engine/Plugins/Experimental/AnimNextControlRig/Source/AnimNextControlRig/Private/AnimNextControlRigHierarchyMappings.cpp
@@ -0,0 +1,339 @@
+// Copyright Epic Games, Inc. All Rights Reserved.
+
+#include "AnimNextControlRigHierarchyMappings.h"
+
+#include "ControlRig.h"
+#include "BoneContainer.h"
+#include "Rigs/RigHierarchy.h"
+#include "Animation/NodeMappingContainer.h"
+#include "Animation/AnimNodeBase.h"
+#include "Animation/AnimCurveTypes.h"
+#include "HAL/IConsoleManager.h"
+#include "Tools/ControlRigIOSettings.h"
+#include "Tools/ControlRigPoseAdapter.h"
+#include "Components/SkeletalMeshComponent.h"
+#include "EvaluationVM/KeyframeState.h"
+#include "AnimNextControlRigPoseAdapter.h"
+#include "Engine/SkeletalMesh.h"
+#include "Animation/Skeleton.h"
+
+namespace UE::AnimNext::ControlRig
+{
+
+void FAnimNextControlRigHierarchyMappings::InitializeInstance()
+{
+	PoseAdapter = MakeShareable(new FAnimNextControlRigPoseAdapter());
+}
+
+void FAnimNextControlRigHierarchyMappings::LinkToHierarchy(URigHierarchy* InHierarchy)
+{
+	if (InHierarchy)
+	{
+		InHierarchy->LinkPoseAdapter(PoseAdapter);
+	}
+}
+
+void FAnimNextControlRigHierarchyMappings::UpdateControlRigRefPoseIfNeeded(UControlRig* ControlRig
+	, const UObject* InstanceObject
+	, const USkeletalMeshComponent* SkeletalMeshComponent
+	, const UE::AnimNext::FReferencePose& InRefPose
+	, bool bInSetRefPoseFromSkeleton
+	, bool bIncludePoseInHash)
+{
+	check(ControlRig);
+
+	if (!bInSetRefPoseFromSkeleton)
+	{
+		return;
+	}
+
+	int32 ExpectedHash = 0;
+	ExpectedHash = HashCombine(ExpectedHash, (int32)reinterpret_cast<uintptr_t>(InstanceObject));
+	ExpectedHash = HashCombine(ExpectedHash, (int32)reinterpret_cast<uintptr_t>(SkeletalMeshComponent));
+
+	if (SkeletalMeshComponent)
+	{
+		ExpectedHash = HashCombine(ExpectedHash, (int32)reinterpret_cast<uintptr_t>(SkeletalMeshComponent->GetSkeletalMeshAsset()));
+	}
+
+	if (bIncludePoseInHash)
+	{
+		const int32 NumRefTransforms = InRefPose.ReferenceLocalTransforms.Num();
+		for (int32 Index = 0; Index < NumRefTransforms; Index++)
+		{
+			const FTransform& Transform = InRefPose.ReferenceLocalTransforms[Index];
+			const FQuat Rotation = Transform.GetRotation();
+
+			ExpectedHash = HashCombine(ExpectedHash, GetTypeHash(Transform.GetTranslation()));
+			ExpectedHash = HashCombine(ExpectedHash, GetTypeHash(FVector4(Rotation.X, Rotation.Y, Rotation.Z, Rotation.W)));
+			ExpectedHash = HashCombine(ExpectedHash, GetTypeHash(Transform.GetScale3D()));
+		}
+	}
+
+	if (RefPoseSetterHash.IsSet() && (ExpectedHash == RefPoseSetterHash.GetValue()))
+	{
+		return;
+	}
+
+	const FReferenceSkeleton* RefSkeleton = FAnimNextControlRigPoseAdapter::GetReferenceSkeleton(InRefPose);
+	check(RefSkeleton != nullptr);
+	ControlRig->SetBoneInitialTransformsFromRefSkeleton(*RefSkeleton);
+
+	RefPoseSetterHash = ExpectedHash;
+}
+
+void FAnimNextControlRigHierarchyMappings::UpdateInputOutputMappingIfRequired(UControlRig* InControlRig
+	, URigHierarchy* InHierarchy
+	, const UE::AnimNext::FReferencePose& InRefPose
+	, int32 InCurrentLOD
+	, const TArray<FBoneReference>& InInputBonesToTransfer
+	, const TArray<FBoneReference>& InOutputBonesToTransfer
+	, const TWeakObjectPtr<UNodeMappingContainer>& InNodeMappingContainer
+	, bool bInTransferPoseInGlobalSpace
+	, bool bInResetInputPoseToInitial)
+{
+	DECLARE_SCOPE_HIERARCHICAL_COUNTER_FUNC()
+
+	if (InHierarchy == nullptr)
+	{
+		return;
+	}
+
+	InHierarchy->UnlinkPoseAdapter();
+
+	if (!PoseAdapter.IsValid())
+	{
+		return;
+	}
+
+	PoseAdapter->UpdateInputOutputMappingIfRequired(InControlRig
+		, InHierarchy
+		, InRefPose
+		, InCurrentLOD
+		, InNodeMappingContainer
+		, bInTransferPoseInGlobalSpace
+		, bInResetInputPoseToInitial);
+}
+
+void FAnimNextControlRigHierarchyMappings::UpdateInput(UControlRig* ControlRig
+	, UE::AnimNext::FKeyframeState& InOutput
+	, const FControlRigIOSettings& InInputSettings
+	, const FControlRigIOSettings& InOutputSettings
+	, TWeakObjectPtr<UNodeMappingContainer>& InNodeMappingContainer
+	, bool bInExecute
+	, bool bInTransferInputPose
+	, bool bInResetInputPoseToInitial
+	, bool bInTransferPoseInGlobalSpace
+	, bool bInTransferInputCurves)
+{
+	DECLARE_SCOPE_HIERARCHICAL_COUNTER_FUNC()
+
+	URigHierarchy* Hierarchy = ControlRig->GetHierarchy();
+	if (Hierarchy == nullptr)
+	{
+		return;
+	}
+
+#if WITH_EDITOR
+	// if we are recording any change - let's clear the undo stack
+	if (bInExecute && Hierarchy->IsTracingChanges())
+	{
+		Hierarchy->ResetTransformStack();
+	}
+#endif
+
+	if (InOutputSettings.bUpdatePose && PoseAdapter)
+	{
+		if (InInputSettings.bUpdatePose)
+		{
+			// reset transforms here to prevent additive transforms from accumulating to INF
+			// we only update transforms from the mesh pose for bones in the current LOD, 
+			// so the reset here ensures excluded bones are also reset
+			if (!PoseAdapter->GetBonesToResetToInitial().IsEmpty())
+			{
+				FRigHierarchyValidityBracket ValidityBracket(Hierarchy);
+				{
+#if WITH_EDITOR
+					// make sure transient controls don't get reset
+					UControlRig::FTransientControlPoseScope PoseScope(ControlRig);
+#endif
+					for (const int32& BoneElementIndex : PoseAdapter->GetBonesToResetToInitial())
+					{
+						if (FRigTransformElement* BoneElement = Hierarchy->Get<FRigBoneElement>(BoneElementIndex))
+						{
+							const FTransform InitialLocalTransform = Hierarchy->GetTransform(BoneElement, ERigTransformType::InitialLocal);
+							BoneElement->GetTransform().Set(ERigTransformType::CurrentLocal, InitialLocalTransform);
+							BoneElement->GetDirtyState().MarkClean(ERigTransformType::CurrentLocal);
+							BoneElement->GetDirtyState().MarkDirty(ERigTransformType::CurrentGlobal);
+						}
+					}
+				}
+			}
+
+			if (bInTransferInputPose && InOutput.Pose.GetNumBones() == PoseAdapter->GetLocalPose().Num())
+			{
+				PoseAdapter->MarkDependentsDirty();
+				PoseAdapter->CopyBonesFrom(InOutput.Pose);
+				PoseAdapter->UpdateDirtyStates(true);
+			}
+		}
+
+		if (InInputSettings.bUpdateCurves || InOutputSettings.bUpdateCurves)
+		{
+			Hierarchy->UnsetCurveValues();
+
+			const TArray<FRigBaseElement*> HierarchyCurves = Hierarchy->GetCurvesFast();
+			PoseAdapter->SetHierarchyCurvesLookup(HierarchyCurves);
+			PoseAdapter->SetPoseCurveToHierarchyCurve(HierarchyCurves, InOutput.Curves);
+		}
+	}
+
+#if WITH_EDITOR
+	if (bInExecute && Hierarchy->IsTracingChanges())
+	{
+		Hierarchy->StorePoseForTrace(TEXT("FAnimNode_ControlRigBase::UpdateInput"));
+	}
+#endif
+}
+
+void FAnimNextControlRigHierarchyMappings::UpdateOutput(UControlRig* ControlRig
+	, UE::AnimNext::FKeyframeState& InOutput
+	, const FControlRigIOSettings& InOutputSettings
+	, TWeakObjectPtr<UNodeMappingContainer>& InNodeMappingContainer
+	, bool bInExecute
+	, bool bInTransferPoseInGlobalSpace)
+{
+	DECLARE_SCOPE_HIERARCHICAL_COUNTER_FUNC()
+
+	URigHierarchy* Hierarchy = ControlRig->GetHierarchy();
+	if (Hierarchy == nullptr)
+	{
+		return;
+	}
+
+	if (InOutputSettings.bUpdatePose && PoseAdapter)
+	{
+		// make sure the local respectively global transforms are all ready
+		if (PoseAdapter->GetTransferInLocalSpace())
+		{
+			PoseAdapter->ConvertToLocalPose();
+		}
+		else
+		{
+			PoseAdapter->ConvertToGlobalPose();
+		}
+
+		// if we transfered in global - let's convert back to local
+		if (!PoseAdapter->GetTransferInLocalSpace())
+		{
+			if (PoseAdapter->GetGlobalPose().Num() == InOutput.Pose.GetNumBones())
+			{
+				PoseAdapter->ConvertToLocalPose();
+			}
+		}
+
+		InOutput.Pose.CopyTransformsFrom(PoseAdapter->GetLocalPose());
+	}
+
+	if (InOutputSettings.bUpdateCurves)
+	{
+		const TArray<int32>& ChangedCurveIndices = Hierarchy->GetChangedCurveIndices();
+		if (ChangedCurveIndices.Num() > 0)
+		{
+			const TArray<FRigBaseElement*> HierarchyCurves = Hierarchy->GetCurvesFast();
+
+			int32 CurveIndex = 0;
+			int32 CurvesCopied = 0;
+			InOutput.Curves.ForEachElement([this, &CurveIndex, &CurvesCopied, &HierarchyCurves](const UE::Anim::FCurveElement& InCurveElement)
+				{
+					// the index stored here is the sub index of the curve (the index of the curve within the list of curves)
+					const int32& HierarchyIndex = PoseAdapter->GetPoseCurveToHierarchyCurve()[CurveIndex];
+					if (HierarchyIndex != INDEX_NONE)
+					{
+						const FRigCurveElement* HierarchyCurve = CastChecked<FRigCurveElement>(HierarchyCurves[HierarchyIndex]);
+						if (HierarchyCurve->IsValueSet())
+						{
+							const_cast<UE::Anim::FCurveElement*>(&InCurveElement)->Value = HierarchyCurve->Get();
+						}
+						CurvesCopied++;
+					}
+					CurveIndex++;
+				});
+
+			if (CurvesCopied < HierarchyCurves.Num())
+			{
+				HierarchyCurveCopied.SetNumUninitialized(Hierarchy->Num());
+				FMemory::Memzero(HierarchyCurveCopied.GetData(), HierarchyCurveCopied.GetAllocatedSize());
+
+				FBlendedCurve ControlRigCurves;
+				ControlRigCurves.Reserve(ChangedCurveIndices.Num());
+				for (const int32& ChangedCurveIndex : ChangedCurveIndices)
+				{
+					if (!HierarchyCurveCopied[ChangedCurveIndex])
+					{
+						if (const FRigCurveElement* HierarchyCurve = Hierarchy->Get<FRigCurveElement>(ChangedCurveIndex))
+						{
+							if (HierarchyCurve->IsValueSet())
+							{
+								ControlRigCurves.Add(HierarchyCurve->GetFName(), HierarchyCurve->Get());
+							}
+						}
+						HierarchyCurveCopied[ChangedCurveIndex] = true;
+					}
+				}
+
+				InOutput.Curves.Combine(ControlRigCurves);
+			}
+		}
+	}
+
+#if WITH_EDITOR
+	if (bInExecute && Hierarchy->IsTracingChanges())
+	{
+		Hierarchy->StorePoseForTrace(TEXT("FAnimNode_ControlRigBase::UpdateOutput"));
+		Hierarchy->DumpTransformStackToFile();
+	}
+#endif
+}
+
+bool FAnimNextControlRigHierarchyMappings::CheckPoseAdapter() const
+{
+	DECLARE_SCOPE_HIERARCHICAL_COUNTER_FUNC()
+
+	if (!PoseAdapter.IsValid())
+	{
+		return false;
+	}
+	return true;
+}
+
+bool FAnimNextControlRigHierarchyMappings::IsUpdateToDate(const URigHierarchy* InHierarchy) const
+{
+	DECLARE_SCOPE_HIERARCHICAL_COUNTER_FUNC()
+
+	if (!PoseAdapter->IsUpdateToDate(InHierarchy))
+	{
+		return false;
+	}
+	return true;
+}
+
+void FAnimNextControlRigHierarchyMappings::PerformUpdateToDate(UControlRig* ControlRig
+	, URigHierarchy* InHierarchy
+	, const UE::AnimNext::FReferencePose& InRefPose
+	, int32 InCurrentLOD
+	, const TWeakObjectPtr<UNodeMappingContainer>& InNodeMappingContainer
+	, bool bInTransferPoseInGlobalSpace
+	, bool bInResetInputPoseToInitial)
+{
+	DECLARE_SCOPE_HIERARCHICAL_COUNTER_FUNC()
+
+	if (!PoseAdapter->IsUpdateToDate(InHierarchy))
+	{
+		InHierarchy->UnlinkPoseAdapter();
+		PoseAdapter->UpdateInputOutputMappingIfRequired(ControlRig, InHierarchy, InRefPose, InCurrentLOD, InNodeMappingContainer, bInTransferPoseInGlobalSpace, bInResetInputPoseToInitial);
+		InHierarchy->LinkPoseAdapter(PoseAdapter);
+	}
+}
+
+} // namespace UE::AnimNext::ControlRig
diff --git a/Engine/Plugins/Experimental/AnimNextControlRig/Source/AnimNextControlRig/Private/AnimNextControlRigPoseAdapter.cpp b/Engine/Plugins/Experimental/AnimNextControlRig/Source/AnimNextControlRig/Private/AnimNextControlRigPoseAdapter.cpp
new file mode 100644
--- /dev/null
+++ b/Engine/Plugins/Experimental/AnimNextControlRig/Source/AnimNextControlRig/Private/AnimNextControlRigPoseAdapter.cpp
@@ -0,0 +1,211 @@
+// Copyright Epic Games, Inc. All Rights Reserved.
+
+#include "AnimNextControlRigPoseAdapter.h"
+#include "Rigs/RigHierarchy.h"
+#include "ControlRig.h"
+#include "Animation/NodeMappingContainer.h"
+#include "Engine/SkeletalMesh.h"
+
+namespace UE::AnimNext::ControlRig
+{
+
+void FAnimNextControlRigPoseAdapter::UpdateInputOutputMappingIfRequired(UControlRig* InControlRig
+	, URigHierarchy* InHierarchy
+	, const UE::AnimNext::FReferencePose& InRefPose
+	, const int32 InCurentLOD
+	, const TWeakObjectPtr<UNodeMappingContainer>& InNodeMappingContainer
+	, bool bInTransferPoseInGlobalSpace
+	, bool bInResetInputPoseToInitial)
+{
+	DECLARE_SCOPE_HIERARCHICAL_COUNTER_FUNC()
+
+	if (InHierarchy == nullptr)
+	{
+		return;
+	}
+
+	ParentPoseIndices.Reset();
+	RequiresHierarchyForSpaceConversion.Reset();
+	ElementIndexToPoseIndex.Reset();
+	PoseIndexToElementIndex.Reset();
+	GlobalPose.Reset();
+	LocalPose.Reset();
+	HierarchyCurveLookup.Reset();
+
+	const int32 NumBonesInPose = InRefPose.GetNumBonesForLOD(InCurentLOD);
+
+	ParentPoseIndices.Reserve(NumBonesInPose);
+	RequiresHierarchyForSpaceConversion.Reserve(NumBonesInPose);
+	GlobalPose.AddDefaulted(NumBonesInPose);
+	LocalPose.AddDefaulted(NumBonesInPose);
+
+	bTransferInLocalSpace = !(bInTransferPoseInGlobalSpace || InNodeMappingContainer.IsValid());
+
+	ParentPoseIndices.SetNumUninitialized(NumBonesInPose);
+	for (int32 Index = 0; Index < NumBonesInPose; Index++)
+	{
+		ParentPoseIndices[Index] = InRefPose.GetLODParentBoneIndex(InCurentLOD, Index);
+	}
+	RequiresHierarchyForSpaceConversion.AddDefaulted(NumBonesInPose);
+
+	UpdateDirtyStates();
+
+	TArray<int32> MappedBoneElementIndices;
+	if (NumBonesInPose > 0)
+	{
+		ElementIndexToPoseIndex.Reserve(NumBonesInPose);
+		PoseIndexToElementIndex.Reserve(NumBonesInPose);
+
+		const FReferenceSkeleton* RefSkeleton = GetReferenceSkeleton(InRefPose);
+		if (const USkeleton* Skeleton = InRefPose.GetSkeletonAsset())
+		{
+			RefSkeleton = &Skeleton->GetReferenceSkeleton();
+		}
+
+		// @todo: thread-safe? probably not in editor, but it may not be a big issue in editor
+		if (InNodeMappingContainer.IsValid())
+		{
+			// get target to source mapping table - this is reversed mapping table
+			TMap<FName, FName> TargetToSourceMappingTable;
+			InNodeMappingContainer->GetTargetToSourceMappingTable(TargetToSourceMappingTable);
+
+			// now fill up node name
+			for (uint16 Index = 0; Index < NumBonesInPose; ++Index)
+			{
+				// get bone name, and find reverse mapping
+				const FSkeletonPoseBoneIndex BoneIndex(InRefPose.GetSkeletonBoneIndexFromLODBoneIndex(Index));
+				if (BoneIndex.IsValid())
+				{
+					const FName TargetNodeName = RefSkeleton->GetBoneName(BoneIndex.GetInt());
+					const FName* SourceName = TargetToSourceMappingTable.Find(TargetNodeName);
+					if (SourceName)
+					{
+						const int32 ElementIndex = InHierarchy->GetIndex({ *SourceName, ERigElementType::Bone });
+						if (ElementIndex != INDEX_NONE)
+						{
+							MappedBoneElementIndices.Add(ElementIndex);
+							ElementIndexToPoseIndex.Add(static_cast<uint16>(ElementIndex), Index);
+							PoseIndexToElementIndex.Add(ElementIndex);
+							LocalPose[Index] = InHierarchy->GetLocalTransform(ElementIndex);
+							GlobalPose[Index] = InHierarchy->GetGlobalTransform(ElementIndex);
+							continue;
+						}
+					}
+				}
+				PoseIndexToElementIndex.Add(INDEX_NONE);
+			}
+		}
+		else
+		{
+			TArray<FName> NodeNames;
+			TArray<FNodeItem> NodeItems;
+			InControlRig->GetMappableNodeData(NodeNames, NodeItems);
+
+			// even if not mapped, we map only node that exists in the control rig
+			for (uint16 Index = 0; Index < NumBonesInPose; ++Index)
+			{
+				const FSkeletonPoseBoneIndex BoneIndex(InRefPose.GetSkeletonBoneIndexFromLODBoneIndex(Index));
+				if (BoneIndex.IsValid())
+				{
+					const FName BoneName = RefSkeleton->GetBoneName(BoneIndex.GetInt());
+					if (NodeNames.Contains(BoneName))
+					{
+						const int32 ElementIndex = InHierarchy->GetIndex({ BoneName, ERigElementType::Bone });
+						if (ElementIndex != INDEX_NONE)
+						{
+							MappedBoneElementIndices.Add(ElementIndex);
+							ElementIndexToPoseIndex.Add(static_cast<uint16>(ElementIndex), Index);
+							PoseIndexToElementIndex.Add(ElementIndex);
+							LocalPose[Index] = InHierarchy->GetLocalTransform(ElementIndex);
+							GlobalPose[Index] = InHierarchy->GetGlobalTransform(ElementIndex);
+							continue;
+						}
+					}
+				}
+				PoseIndexToElementIndex.Add(INDEX_NONE);
+			}
+		}
+
+		// once we know all of the bones we are going to transfer - we can check if any of these bones has a different parenting
+		// relationship in the skeleton used in the anim graph vs the hierarchy in the rig. in that case we have to transfer in global
+		if (bTransferInLocalSpace)
+		{
+			for (const int32& BoneElementIndex : MappedBoneElementIndices)
+			{
+				const int32 HierarchyParentIndex = InHierarchy->GetFirstParent(BoneElementIndex);
+				const int16 PoseIndex = ElementIndexToPoseIndex.FindChecked(BoneElementIndex);
+				const FCompactPoseBoneIndex CompactPoseParentIndex(ParentPoseIndices[PoseIndex]);
+
+				FName HierarchyParentName(NAME_None);
+				FName PoseParentName(NAME_None);
+
+				if (HierarchyParentIndex != INDEX_NONE)
+				{
+					HierarchyParentName = InHierarchy->Get(HierarchyParentIndex)->GetFName();
+				}
+				if (CompactPoseParentIndex.IsValid())
+				{
+					const FSkeletonPoseBoneIndex SkeletonIndex(InRefPose.GetSkeletonBoneIndexFromLODBoneIndex(CompactPoseParentIndex.GetInt()));
+					PoseParentName = RefSkeleton->GetBoneName(SkeletonIndex.GetInt());
+				}
+
+				if (HierarchyParentName.IsEqual(PoseParentName, ENameCase::CaseSensitive))
+				{
+					continue;
+				}
+
+				RequiresHierarchyForSpaceConversion[PoseIndex] = true;
+				check(PoseIndexToElementIndex[PoseIndex] != INDEX_NONE);
+				bTransferInLocalSpace = false;
+			}
+		}
+
+		// only reset the full pose if we are not mapping all bones
+		const TArray<FRigBaseElement*>& HierarchyBones = InHierarchy->GetBonesFast();
+		const bool bMapsAllBones = MappedBoneElementIndices.Num() == HierarchyBones.Num();
+		BonesToResetToInitial.Reset();
+		bRequiresResetPoseToInitial = bInResetInputPoseToInitial && !bMapsAllBones;
+
+		if (bRequiresResetPoseToInitial)
+		{
+			BonesToResetToInitial.Reserve(HierarchyBones.Num() - MappedBoneElementIndices.Num());
+
+			// bone is mapped stores sub indices (bone index within the list of bones)
+			TArray<bool> BoneIsMapped;
+			BoneIsMapped.AddZeroed(HierarchyBones.Num());
+			for (const int32& MappedTransformIndex : MappedBoneElementIndices)
+			{
+				const FRigBaseElement* MappedElement = InHierarchy->Get(MappedTransformIndex);
+				check(MappedElement);
+				BoneIsMapped[MappedElement->GetSubIndex()] = true;
+			}
+
+			// when we want to know which bones to reset we want to convert back to a global index
+			for (int32 UnmappedBoneIndex = 0; UnmappedBoneIndex < BoneIsMapped.Num(); UnmappedBoneIndex++)
+			{
+				if (!BoneIsMapped[UnmappedBoneIndex])
+				{
+					BonesToResetToInitial.Add(HierarchyBones[UnmappedBoneIndex]->GetIndex());
+				}
+			}
+		}
+	}
+}
+
+const FReferenceSkeleton* FAnimNextControlRigPoseAdapter::GetReferenceSkeleton(const UE::AnimNext::FReferencePose& InRefPose)
+{
+	const FReferenceSkeleton* ReferenceSkeleton = nullptr;
+
+	if (const USkeletalMesh* SkeletalMesh = InRefPose.GetSkeletalMeshAsset())
+	{
+		ReferenceSkeleton = &SkeletalMesh->GetRefSkeleton();
+	}
+	else if (const USkeleton* Skeleton = InRefPose.GetSkeletonAsset())
+	{
+		ReferenceSkeleton = &Skeleton->GetReferenceSkeleton();
+	}
+
+	return ReferenceSkeleton;
+}
+
+} // namespace UE::AnimNext::ControlRig
diff --git a/Engine/Plugins/Experimental/AnimNextControlRig/Source/AnimNextControlRig/Private/ControlRigTask.cpp b/Engine/Plugins/Experimental/AnimNextControlRig/Source/AnimNextControlRig/Private/ControlRigTask.cpp
--- a/Engine/Plugins/Experimental/AnimNextControlRig/Source/AnimNextControlRig/Private/ControlRigTask.cpp
+++ b/Engine/Plugins/Experimental/AnimNextControlRig/Source/AnimNextControlRig/Private/ControlRigTask.cpp
@@ -70,33 +70,20 @@ void FAnimNextControlRigTask::ExecuteControlRig(UE::AnimNext::FEvaluationVM& VM,
 
 		const UE::AnimNext::FLODPoseStack& Pose = KeyFrameState.Pose;
 		const UE::AnimNext::FReferencePose& RefPose = Pose.GetRefPose();
-		const USkeletalMesh* SkeletalMesh = RefPose.SkeletalMesh.Get();
-
-		// TODO : hack to get things rolling, as we don't have AnimInstanceProxy to obtain RequiredBones, until ControlRig helpers can use anim next data
-		FBoneContainer RequiredBones;
-		TArray<FBoneIndexType> LODBoneIndexToMeshBoneIndex(RefPose.GetLODBoneIndexToMeshBoneIndexMap(VM.GetCurrentLOD()));
-		LODBoneIndexToMeshBoneIndex.Sort();
-		RequiredBones.InitializeTo(LODBoneIndexToMeshBoneIndex, InstanceData->SkeletalMeshComponent->GetCurveFilterSettings(), *SkeletalMesh);
-		RequiredBones.SetRefPoseOverride(InstanceData->SkeletalMeshComponent->GetRefPoseOverride());
-
-		FPoseContext Output(RequiredBones);
-
-		UE::AnimNext::FGenerationTools::RemapPose(Pose, Output);
-		Output.Curve.CopyFrom(KeyFrameState.Curves);
-		UE::AnimNext::FGenerationTools::RemapAttributes(Pose, KeyFrameState.Attributes, Output);
 
 		// temporarily give control rig access to the stack allocated attribute container
 		// control rig may have rig units that can add/get attributes to/from this container
-		UControlRig::FAnimAttributeContainerPtrScope AttributeScope(ControlRig, Output.CustomAttributes);
+		UControlRig::FAnimAttributeContainerPtrScope AttributeScope(ControlRig, KeyFrameState.Attributes);
 
 		if (InstanceData->bUpdateInputOutputMapping)
 		{
 			// UpdateInputOutputMappingIfRequired was done in CacheBones, but there is no AnimNext equivalent
 			InstanceData->ControlRigHierarchyMappings.UpdateInputOutputMappingIfRequired(ControlRig
 				, ControlRig->GetHierarchy()
-				, Output.Pose.GetBoneContainer()
-				, TArray<FBoneReference>() //SharedData->InputBonesToTransfer // TODO zzz : Unsuported in AnimNext
-				, TArray<FBoneReference>() //SharedData->OutputBonesToTransfer // TODO zzz : Unsuported in AnimNext
+				, RefPose
+				, VM.GetCurrentLOD()
+				, TArray<FBoneReference>()
+				, TArray<FBoneReference>()
 				, InstanceData->NodeMappingContainer
 				, SharedData->bTransferPoseInGlobalSpace
 				, SharedData->bResetInputPoseToInitial);
@@ -104,21 +91,19 @@ void FAnimNextControlRigTask::ExecuteControlRig(UE::AnimNext::FEvaluationVM& VM,
 			InstanceData->bUpdateInputOutputMapping = false;
 		}
 
-		if (InstanceData->ControlRigHierarchyMappings.IsPoseAdapterEnabled())
+		if (!InstanceData->ControlRigHierarchyMappings.IsUpdateToDate(Hierarchy))
 		{
-			if (!InstanceData->ControlRigHierarchyMappings.IsUpdateToDate(Hierarchy))
-			{
-				InstanceData->ControlRigHierarchyMappings.PerformUpdateToDate(ControlRig
-					, Hierarchy
-					, Output.Pose.GetBoneContainer()
-					, InstanceData->NodeMappingContainer
-					, SharedData->bTransferPoseInGlobalSpace
-					, SharedData->bResetInputPoseToInitial);
-			}
+			InstanceData->ControlRigHierarchyMappings.PerformUpdateToDate(ControlRig
+				, Hierarchy
+				, RefPose
+				, VM.GetCurrentLOD()
+				, InstanceData->NodeMappingContainer
+				, SharedData->bTransferPoseInGlobalSpace
+				, SharedData->bResetInputPoseToInitial);
 		}
 
 		// first update input to the system
-		UpdateInput(VM, KeyFrameState, ControlRig, Output);
+		UpdateInput(VM, KeyFrameState, ControlRig, KeyFrameState);
 
 		if (InstanceData->bExecute)
 		{
@@ -183,15 +168,11 @@ void FAnimNextControlRigTask::ExecuteControlRig(UE::AnimNext::FEvaluationVM& VM,
 		}
 
 		// now update output
-		UpdateOutput(VM, KeyFrameState, ControlRig, Output);
-
-		UE::AnimNext::FGenerationTools::RemapPose(Output, KeyFrameState.Pose);
-		KeyFrameState.Curves.CopyFrom(Output.Curve);
-		UE::AnimNext::FGenerationTools::RemapAttributes(Output, KeyFrameState.Pose, KeyFrameState.Attributes);
+		UpdateOutput(VM, KeyFrameState, ControlRig, KeyFrameState);
 	}
 }
 
-void FAnimNextControlRigTask::UpdateInput(UE::AnimNext::FEvaluationVM& VM, UE::AnimNext::FKeyframeState& KeyFrameState, UControlRig* ControlRig, FPoseContext& InOutput) const
+void FAnimNextControlRigTask::UpdateInput(UE::AnimNext::FEvaluationVM& VM, UE::AnimNext::FKeyframeState& KeyFrameState, UControlRig* ControlRig, UE::AnimNext::FKeyframeState& InOutput) const
 {
 	DECLARE_SCOPE_HIERARCHICAL_COUNTER_FUNC()
 
@@ -213,10 +194,10 @@ void FAnimNextControlRigTask::UpdateInput(UE::AnimNext::FEvaluationVM& VM, UE::A
 		, SharedData->bTransferPoseInGlobalSpace
 		, SharedData->bTransferInputCurves);
 
-	InstanceData->ControlRigVariableMappings.UpdateCurveInputs(InstanceData->ControlRig, SharedData->InputMapping, InOutput.Curve);
+	InstanceData->ControlRigVariableMappings.UpdateCurveInputs(InstanceData->ControlRig, SharedData->InputMapping, InOutput.Curves);
 }
 
-void FAnimNextControlRigTask::UpdateOutput(UE::AnimNext::FEvaluationVM& VM, UE::AnimNext::FKeyframeState& KeyFrameState, UControlRig* ControlRig, FPoseContext& InOutput) const
+void FAnimNextControlRigTask::UpdateOutput(UE::AnimNext::FEvaluationVM& VM, UE::AnimNext::FKeyframeState& KeyFrameState, UControlRig* ControlRig, UE::AnimNext::FKeyframeState& InOutput) const
 {
 	DECLARE_SCOPE_HIERARCHICAL_COUNTER_FUNC()
 
@@ -234,7 +215,7 @@ void FAnimNextControlRigTask::UpdateOutput(UE::AnimNext::FEvaluationVM& VM, UE::
 		, InstanceData->bExecute
 		, SharedData->bTransferPoseInGlobalSpace);
 
-	InstanceData->ControlRigVariableMappings.UpdateCurveOutputs(InstanceData->ControlRig, SharedData->OutputMapping, InOutput.Curve);
+	InstanceData->ControlRigVariableMappings.UpdateCurveOutputs(InstanceData->ControlRig, SharedData->OutputMapping, InOutput.Curves);
 }
 
 bool FAnimNextControlRigTask::CanExecute(const UControlRig* ControlRig) const
diff --git a/Engine/Plugins/Experimental/AnimNextControlRig/Source/AnimNextControlRig/Private/ControlRigTask.h b/Engine/Plugins/Experimental/AnimNextControlRig/Source/AnimNextControlRig/Private/ControlRigTask.h
--- a/Engine/Plugins/Experimental/AnimNextControlRig/Source/AnimNextControlRig/Private/ControlRigTask.h
+++ b/Engine/Plugins/Experimental/AnimNextControlRig/Source/AnimNextControlRig/Private/ControlRigTask.h
@@ -11,6 +11,11 @@
 // --- ---
 #include "ControlRigTask.generated.h"
 
+namespace UE::AnimNext
+{
+struct FKeyframeState;
+}
+
 DECLARE_CYCLE_STAT_EXTERN(TEXT("AnimNext Task: ControlRig"), STAT_AnimNext_Task_ControlRig, STATGROUP_AnimNext, ANIMNEXTCONTROLRIG_API);
 
 USTRUCT()
@@ -27,8 +32,8 @@ struct ANIMNEXTCONTROLRIG_API FAnimNextControlRigTask : public FAnimNextEvaluati
 
 private:
 	void ExecuteControlRig(UE::AnimNext::FEvaluationVM& VM, UE::AnimNext::FKeyframeState& KeyFrameState, UControlRig* ControlRig) const;
-	void UpdateInput(UE::AnimNext::FEvaluationVM& VM, UE::AnimNext::FKeyframeState& KeyFrameState, UControlRig* ControlRig, FPoseContext& InOutput) const;
-	void UpdateOutput(UE::AnimNext::FEvaluationVM& VM, UE::AnimNext::FKeyframeState& KeyFrameState, UControlRig* ControlRig, FPoseContext& InOutput) const;
+	void UpdateInput(UE::AnimNext::FEvaluationVM& VM, UE::AnimNext::FKeyframeState& KeyFrameState, UControlRig* ControlRig, UE::AnimNext::FKeyframeState& InOutput) const;
+	void UpdateOutput(UE::AnimNext::FEvaluationVM& VM, UE::AnimNext::FKeyframeState& KeyFrameState, UControlRig* ControlRig, UE::AnimNext::FKeyframeState& InOutput) const;
 
 	bool CanExecute(const UControlRig* ControlRig) const;
 
