Fixed AnimNext Control Rig not setting the owner component in the data registry, so transforms From / To World work now

diff --git a/Engine/Plugins/Experimental/AnimNextControlRig/Source/AnimNextControlRig/AnimNextControlRig.Build.cs b/Engine/Plugins/Experimental/AnimNextControlRig/Source/AnimNextControlRig/AnimNextControlRig.Build.cs
--- a/Engine/Plugins/Experimental/AnimNextControlRig/Source/AnimNextControlRig/AnimNextControlRig.Build.cs
+++ b/Engine/Plugins/Experimental/AnimNextControlRig/Source/AnimNextControlRig/AnimNextControlRig.Build.cs
@@ -30,6 +30,7 @@ namespace UnrealBuildTool.Rules
 						"Slate",
 						"AnimNextUncookedOnly",
 						"RigVMDeveloper",
+						"AnimationCore",
 				   });
 			}
 		}
diff --git a/Engine/Plugins/Experimental/AnimNextControlRig/Source/AnimNextControlRig/Private/ControlRigTrait.cpp b/Engine/Plugins/Experimental/AnimNextControlRig/Source/AnimNextControlRig/Private/ControlRigTrait.cpp
--- a/Engine/Plugins/Experimental/AnimNextControlRig/Source/AnimNextControlRig/Private/ControlRigTrait.cpp
+++ b/Engine/Plugins/Experimental/AnimNextControlRig/Source/AnimNextControlRig/Private/ControlRigTrait.cpp
@@ -17,6 +17,7 @@
 #include "Components/SkeletalMeshComponent.h"
 #include "Engine/SkeletalMesh.h"
 #include "Animation/AnimNodeBase.h"
+#include "AnimationDataSource.h"
 #if WITH_EDITOR
 #include "RigVMModel/RigVMController.h"
 #include "ControlRigBlueprint.h" // to get the preview skeleton
@@ -582,6 +583,10 @@ bool FControlRigTrait::SetupControlRig(UObject* InAnimContext, UClass* InControl
 			
 			InstanceData->ControlRigHierarchyMappings.ResetRefPoseSetterHash();
 			
+			// register bindable object as data source (used for To World / From World transformations)
+			UObject* BindableObject = FControlRigObjectBinding::GetBindableObject(InAnimContext);
+			ControlRig->GetDataSourceRegistry()->RegisterDataSource(UControlRig::OwnerComponent, BindableObject);
+
 			//UpdateInputOutputMappingIfRequired(ControlRig, InRequiredBones);
 		}
 	}
