Anim Next Control Rig : Fixed crash when there is a LOD change by updating mappings

#rnx

diff --git a/Engine/Plugins/Experimental/AnimNextControlRig/Source/AnimNextControlRig/Internal/ControlRigTrait.h b/Engine/Plugins/Experimental/AnimNextControlRig/Source/AnimNextControlRig/Internal/ControlRigTrait.h
--- a/Engine/Plugins/Experimental/AnimNextControlRig/Source/AnimNextControlRig/Internal/ControlRigTrait.h
+++ b/Engine/Plugins/Experimental/AnimNextControlRig/Source/AnimNextControlRig/Internal/ControlRigTrait.h
@@ -204,6 +204,8 @@ struct FControlRigTrait : FBaseTrait, IEvaluate, IUpdate, IHierarchy, IGarbageCo
 		FDelegateHandle OnObjectsReinstancedHandle;
 		FDelegateHandle OnInitializedHandle;
 
+		int32 LastLOD = INDEX_NONE;
+
 		bool bExecute = true;
 		bool bClearEventQueueRequired = false;
 		bool bUpdateInputOutputMapping = false;
diff --git a/Engine/Plugins/Experimental/AnimNextControlRig/Source/AnimNextControlRig/Private/ControlRigTask.cpp b/Engine/Plugins/Experimental/AnimNextControlRig/Source/AnimNextControlRig/Private/ControlRigTask.cpp
--- a/Engine/Plugins/Experimental/AnimNextControlRig/Source/AnimNextControlRig/Private/ControlRigTask.cpp
+++ b/Engine/Plugins/Experimental/AnimNextControlRig/Source/AnimNextControlRig/Private/ControlRigTask.cpp
@@ -12,6 +12,7 @@
 #include "GenerationTools.h"
 #include "Engine/SkeletalMesh.h"
 #include "Animation/AnimNodeBase.h"
+#include "Units/Execution/RigUnit_PrepareForExecution.h"
 #if ENABLE_ANIM_DEBUG
 #include "AnimNode_ControlRigBase.h"
 #endif
@@ -48,6 +49,8 @@ void FAnimNextControlRigTask::Execute(UE::AnimNext::FEvaluationVM& VM) const
 
 	// Push our blended result back
 	VM.PushValue(KEYFRAME_STACK_NAME, MoveTemp(KeyframeOut));
+
+	InstanceData->LastLOD = VM.GetCurrentLOD();
 }
 
 void FAnimNextControlRigTask::ExecuteControlRig(UE::AnimNext::FEvaluationVM& VM, UE::AnimNext::FKeyframeState& KeyFrameState, UControlRig* ControlRig) const
@@ -75,13 +78,23 @@ void FAnimNextControlRigTask::ExecuteControlRig(UE::AnimNext::FEvaluationVM& VM,
 		// control rig may have rig units that can add/get attributes to/from this container
 		UControlRig::FAnimAttributeContainerPtrScope AttributeScope(ControlRig, KeyFrameState.Attributes);
 
-		if (InstanceData->bUpdateInputOutputMapping)
+		const int32 CurrentLOD = VM.GetCurrentLOD();
+		const bool bIsLODChange = InstanceData->LastLOD != CurrentLOD;
+		const bool bRunPrepareForExecution = bIsLODChange /* || (RefPoseBonesSerialNumber != LastRefPoseBonesSerialNumber)*/; // TODO : Detect RefPose changes
+		if (InstanceData->bUpdateInputOutputMapping || bRunPrepareForExecution)
 		{
+			if (ControlRig->IsConstructionModeEnabled() ||
+				(ControlRig->IsConstructionRequired() && (InstanceData->bUpdateInputOutputMapping || bRunPrepareForExecution)))
+			{
+				//UpdateGetAssetUserDataDelegate(ControlRig);
+				ControlRig->Execute(FRigUnit_PrepareForExecution::EventName);
+			}
+
 			// UpdateInputOutputMappingIfRequired was done in CacheBones, but there is no AnimNext equivalent
 			InstanceData->ControlRigHierarchyMappings.UpdateInputOutputMappingIfRequired(ControlRig
 				, ControlRig->GetHierarchy()
 				, RefPose
-				, VM.GetCurrentLOD()
+				, CurrentLOD
 				, TArray<FBoneReference>()
 				, TArray<FBoneReference>()
 				, InstanceData->NodeMappingContainer
@@ -96,7 +109,7 @@ void FAnimNextControlRigTask::ExecuteControlRig(UE::AnimNext::FEvaluationVM& VM,
 			InstanceData->ControlRigHierarchyMappings.PerformUpdateToDate(ControlRig
 				, Hierarchy
 				, RefPose
-				, VM.GetCurrentLOD()
+				, CurrentLOD
 				, InstanceData->NodeMappingContainer
 				, SharedData->bTransferPoseInGlobalSpace
 				, SharedData->bResetInputPoseToInitial);
