Control Rig Anim Next Pose Adapter: ensure that the parent's bone index is safe to use

diff --git a/Engine/Plugins/Experimental/AnimNextControlRig/Source/AnimNextControlRig/Private/AnimNextControlRigPoseAdapter.cpp b/Engine/Plugins/Experimental/AnimNextControlRig/Source/AnimNextControlRig/Private/AnimNextControlRigPoseAdapter.cpp
--- a/Engine/Plugins/Experimental/AnimNextControlRig/Source/AnimNextControlRig/Private/AnimNextControlRigPoseAdapter.cpp
+++ b/Engine/Plugins/Experimental/AnimNextControlRig/Source/AnimNextControlRig/Private/AnimNextControlRigPoseAdapter.cpp
@@ -146,7 +146,10 @@ void FAnimNextControlRigPoseAdapter::UpdateInputOutputMappingIfRequired(UControl
 				if (CompactPoseParentIndex.IsValid())
 				{
 					const FSkeletonPoseBoneIndex SkeletonIndex(InRefPose.GetSkeletonBoneIndexFromLODBoneIndex(CompactPoseParentIndex.GetInt()));
-					PoseParentName = RefSkeleton->GetBoneName(SkeletonIndex.GetInt());
+					if (SkeletonIndex.IsValid() && RefSkeleton->IsValidIndex(SkeletonIndex.GetInt()))
+					{
+						PoseParentName = RefSkeleton->GetBoneName(SkeletonIndex.GetInt());
+					}
 				}
 
 				if (HierarchyParentName.IsEqual(PoseParentName, ENameCase::CaseSensitive))
