AnimNext passing BindableObject to graph Update and Evaluate
Created context data structs to call UpdateGraph and EvaluateGraph functions
Updated Control Rig Trait to use the bindable object from the context

diff --git a/Engine/Plugins/Experimental/AnimNext/Source/AnimNext/Private/GenerationTools.cpp b/Engine/Plugins/Experimental/AnimNext/Source/AnimNext/Private/GenerationTools.cpp
--- a/Engine/Plugins/Experimental/AnimNext/Source/AnimNext/Private/GenerationTools.cpp
+++ b/Engine/Plugins/Experimental/AnimNext/Source/AnimNext/Private/GenerationTools.cpp
@@ -114,6 +114,7 @@ struct FCompareBoneIndexType
 				const USkeleton* Skeleton = SkeletalMesh->GetSkeleton();
 
 				OutAnimationReferencePose.GenerationFlags = EReferencePoseGenerationFlags::FastPath;
+				OutAnimationReferencePose.SkeletalMeshComponent = SkeletalMeshComponent;
 				OutAnimationReferencePose.SkeletalMesh = SkeletalMesh;
 				OutAnimationReferencePose.Skeleton = Skeleton;
 
diff --git a/Engine/Plugins/Experimental/AnimNext/Source/AnimNext/Public/ReferencePose.h b/Engine/Plugins/Experimental/AnimNext/Source/AnimNext/Public/ReferencePose.h
--- a/Engine/Plugins/Experimental/AnimNext/Source/AnimNext/Public/ReferencePose.h
+++ b/Engine/Plugins/Experimental/AnimNext/Source/AnimNext/Public/ReferencePose.h
@@ -67,6 +67,7 @@ struct TReferencePose
 	// Mapping of mesh bone indices to mesh parent indices for each bone
 	TArray<FBoneIndexType, AllocatorType> MeshBoneIndexToParentMeshBoneIndexMap;
 
+	TWeakObjectPtr<const USkeletalMeshComponent> SkeletalMeshComponent = nullptr;
 	TWeakObjectPtr<const USkeletalMesh> SkeletalMesh = nullptr;
 	TWeakObjectPtr<const USkeleton> Skeleton = nullptr;
 	EReferencePoseGenerationFlags GenerationFlags = EReferencePoseGenerationFlags::None;
diff --git a/Engine/Plugins/Experimental/AnimNextAnimGraph/Source/AnimNextAnimGraph/Private/Graph/AnimGraph/AnimNode_AnimNextGraph.cpp b/Engine/Plugins/Experimental/AnimNextAnimGraph/Source/AnimNextAnimGraph/Private/Graph/AnimGraph/AnimNode_AnimNextGraph.cpp
--- a/Engine/Plugins/Experimental/AnimNextAnimGraph/Source/AnimNextAnimGraph/Private/Graph/AnimGraph/AnimNode_AnimNextGraph.cpp
+++ b/Engine/Plugins/Experimental/AnimNextAnimGraph/Source/AnimNextAnimGraph/Private/Graph/AnimGraph/AnimNode_AnimNextGraph.cpp
@@ -61,9 +61,9 @@ void FAnimNode_AnimNextGraph::Update_AnyThread(const FAnimationUpdateContext& Co
 
 		PropagateInputProperties(Context.AnimInstanceProxy->GetAnimInstanceObject());
 
-		UE::AnimNext::FTraitEventList InputEventList;
-		UE::AnimNext::FTraitEventList OutputEventList;
-		UE::AnimNext::UpdateGraph(*GraphInstance.Get(), Context.GetDeltaTime(), InputEventList, OutputEventList);
+		FUpdateGraphContext UpdateGraphContext(*GraphInstance.Get(), Context.GetDeltaTime());
+		UpdateGraphContext.SetBindingObject(Context.AnimInstanceProxy->GetSkelMeshComponent());
+		UE::AnimNext::UpdateGraph(UpdateGraphContext);
 	}
 
 	FAnimNode_CustomProperty::Update_AnyThread(Context);
@@ -124,7 +124,9 @@ void FAnimNode_AnimNextGraph::Evaluate_AnyThread(FPoseContext& Output)
 		ResultPose.LODPose = FLODPoseHeap(RefPose, LODLevel, true, Output.ExpectsAdditivePose());
 
 		{
-			const FEvaluationProgram EvaluationProgram = UE::AnimNext::EvaluateGraph(*GraphInstance.Get());
+			UE::AnimNext::FEvaluateGraphContext EvaluateGraphContext(*GraphInstance.Get(), RefPose, LODLevel);
+			EvaluateGraphContext.SetBindingObject(SkeletalMeshComponent);
+			const FEvaluationProgram EvaluationProgram = UE::AnimNext::EvaluateGraph(EvaluateGraphContext);
 
 			FEvaluationVM EvaluationVM(EEvaluationFlags::All, RefPose, LODLevel);
 			bool bHasValidOutput = false;
diff --git a/Engine/Plugins/Experimental/AnimNextAnimGraph/Source/AnimNextAnimGraph/Private/Graph/AnimNextAnimationGraph.cpp b/Engine/Plugins/Experimental/AnimNextAnimGraph/Source/AnimNextAnimGraph/Private/Graph/AnimNextAnimationGraph.cpp
--- a/Engine/Plugins/Experimental/AnimNextAnimGraph/Source/AnimNextAnimGraph/Private/Graph/AnimNextAnimationGraph.cpp
+++ b/Engine/Plugins/Experimental/AnimNextAnimGraph/Source/AnimNextAnimGraph/Private/Graph/AnimNextAnimationGraph.cpp
@@ -1,4 +1,4 @@
-﻿// Copyright Epic Games, Inc. All Rights Reserved.
+// Copyright Epic Games, Inc. All Rights Reserved.
 
 #include "Graph/AnimNextAnimationGraph.h"
 #include "RigVMCore/RigVMExecuteContext.h"
@@ -16,6 +16,8 @@
 #include "UObject/LinkerLoad.h"
 #include "UObject/ObjectResource.h"
 #include "Module/AnimNextModuleInstance.h"
+#include "DataRegistryTypes.h"
+#include "Module/AnimNextSkeletalMeshComponentReferenceComponent.h"
 
 #if WITH_EDITOR	
 #include "Engine/ExternalAssetDependencyGatherer.h"
@@ -32,7 +34,7 @@ UAnimNextAnimationGraph::UAnimNextAnimationGraph(const FObjectInitializer& Objec
 	ExtendedExecuteContext.SetContextPublicDataStruct(FAnimNextExecuteContext::StaticStruct());
 }
 
-TSharedPtr<FAnimNextGraphInstance> UAnimNextAnimationGraph::AllocateInstance(FAnimNextModuleInstance* InModuleInstance, FAnimNextGraphInstance* InParentGraphInstance, FName InEntryPoint) const
+TSharedPtr<FAnimNextGraphInstance> UAnimNextAnimationGraph::AllocateInstance(FAnimNextModuleInstance* InModuleInstance, UE::AnimNext::FExecutionContext* InParentContext, FAnimNextGraphInstance* InParentGraphInstance, FName InEntryPoint) const
 {
 	SCOPE_CYCLE_COUNTER(STAT_AnimNext_Graph_AllocateInstance);
 
@@ -81,6 +83,19 @@ TSharedPtr<FAnimNextGraphInstance> UAnimNextAnimationGraph::AllocateInstance(FAn
 
 	{
 		UE::AnimNext::FExecutionContext Context(*InstanceImpl.Get());
+
+		if (InParentContext != nullptr)
+		{
+			Context.SetBindingObject(InParentContext->GetBindingObject());
+		}
+		else if (InModuleInstance != nullptr)
+		{
+			const FAnimNextSkeletalMeshComponentReferenceComponent& ComponentReference = InModuleInstance->GetComponent<FAnimNextSkeletalMeshComponentReferenceComponent>();
+
+			UE::AnimNext::FDataHandle RefPoseHandle = UE::AnimNext::FDataRegistry::Get()->GetOrGenerateReferencePose(ComponentReference.GetComponent());
+			const UE::AnimNext::FReferencePose& RefPose = RefPoseHandle.GetRef<UE::AnimNext::FReferencePose>();
+			Context.SetBindingObject(RefPose.SkeletalMeshComponent);
+		}
 		InstanceImpl->GraphInstancePtr = Context.AllocateNodeInstance(*InstanceImpl, ResolvedRootTraitHandle);
 	}
 
diff --git a/Engine/Plugins/Experimental/AnimNextAnimGraph/Source/AnimNextAnimGraph/Private/Graph/AnimNextModuleAnimGraphComponent.cpp b/Engine/Plugins/Experimental/AnimNextAnimGraph/Source/AnimNextAnimGraph/Private/Graph/AnimNextModuleAnimGraphComponent.cpp
--- a/Engine/Plugins/Experimental/AnimNextAnimGraph/Source/AnimNextAnimGraph/Private/Graph/AnimNextModuleAnimGraphComponent.cpp
+++ b/Engine/Plugins/Experimental/AnimNextAnimGraph/Source/AnimNextAnimGraph/Private/Graph/AnimNextModuleAnimGraphComponent.cpp
@@ -9,7 +9,7 @@ TWeakPtr<FAnimNextGraphInstance> FAnimNextModuleAnimGraphComponent::AllocateInst
 {
 	check(InAnimationGraph);
 
-	TSharedPtr<FAnimNextGraphInstance> NewInstance = InAnimationGraph->AllocateInstance(&GetModuleInstance(), InParentInstance, InEntryPoint);
+	TSharedPtr<FAnimNextGraphInstance> NewInstance = InAnimationGraph->AllocateInstance(&GetModuleInstance(), nullptr, InParentInstance, InEntryPoint);
 	if(!NewInstance.IsValid())
 	{
 		return TWeakPtr<FAnimNextGraphInstance>();
diff --git a/Engine/Plugins/Experimental/AnimNextAnimGraph/Source/AnimNextAnimGraph/Private/Graph/RigUnit_AnimNextRunAnimationGraph_v1.cpp b/Engine/Plugins/Experimental/AnimNextAnimGraph/Source/AnimNextAnimGraph/Private/Graph/RigUnit_AnimNextRunAnimationGraph_v1.cpp
--- a/Engine/Plugins/Experimental/AnimNextAnimGraph/Source/AnimNextAnimGraph/Private/Graph/RigUnit_AnimNextRunAnimationGraph_v1.cpp
+++ b/Engine/Plugins/Experimental/AnimNextAnimGraph/Source/AnimNextAnimGraph/Private/Graph/RigUnit_AnimNextRunAnimationGraph_v1.cpp
@@ -94,8 +94,9 @@ FRigUnit_AnimNextRunAnimationGraph_v1_Execute()
 	ensure(Result.LODPose.LODLevel == LOD);
 
 	// Every graph in a schedule will see the same input events (if they were queued before the schedule started)
-	FTraitEventList InputEventList;
-	FTraitEventList OutputEventList;
+	FUpdateGraphContext UpdateGraphContext(InstanceRef.Get(), ExecuteContext.GetDeltaTime());
+	UpdateGraphContext.SetBindingObject(RefPose.SkeletalMeshComponent);
+	FTraitEventList& InputEventList = UpdateGraphContext.GetInputEventList();
 
 	// A schedule can contain multiple graphs, we copy the input event list since it might be appended to during our update
 	{
@@ -114,8 +115,11 @@ FRigUnit_AnimNextRunAnimationGraph_v1_Execute()
 	// This reduces churn internally by avoiding a chunk to be repeatedly allocated and freed as we push/pop marks
 	MemStack.Alloc(size_t(FPageAllocator::SmallPageSize) + 1, 16);
 
-	UpdateGraph(InstanceRef.Get(), ExecuteContext.GetDeltaTime(), InputEventList, OutputEventList);
-	EvaluateGraph(InstanceRef.Get(), RefPose, LOD, Result);
+	UpdateGraph(UpdateGraphContext);
+	
+	FEvaluateGraphContext EvaluateGraphContext(InstanceRef.Get(), RefPose, LOD);
+	EvaluateGraphContext.SetBindingObject(RefPose.SkeletalMeshComponent);
+	EvaluateGraph(EvaluateGraphContext, Result);
 
 	// We might have appended new input/output events, append them
 	{
@@ -134,6 +138,6 @@ FRigUnit_AnimNextRunAnimationGraph_v1_Execute()
 		}
 
 		// Append our output events
-		ModuleInstance.OutputEventList.Append(OutputEventList);
+		ModuleInstance.OutputEventList.Append(UpdateGraphContext.GetOutputEventList());
 	}
 }
diff --git a/Engine/Plugins/Experimental/AnimNextAnimGraph/Source/AnimNextAnimGraph/Private/Graph/RigUnit_AnimNextRunAnimationGraph_v2.cpp b/Engine/Plugins/Experimental/AnimNextAnimGraph/Source/AnimNextAnimGraph/Private/Graph/RigUnit_AnimNextRunAnimationGraph_v2.cpp
--- a/Engine/Plugins/Experimental/AnimNextAnimGraph/Source/AnimNextAnimGraph/Private/Graph/RigUnit_AnimNextRunAnimationGraph_v2.cpp
+++ b/Engine/Plugins/Experimental/AnimNextAnimGraph/Source/AnimNextAnimGraph/Private/Graph/RigUnit_AnimNextRunAnimationGraph_v2.cpp
@@ -121,8 +121,10 @@ FRigUnit_AnimNextRunAnimationGraph_v2_Execute()
 	AnimNextExecuteContext.SetDeltaTime(ExecuteContext.GetDeltaTime());
 
 	// Every graph in a schedule will see the same input events (if they were queued before the schedule started)
-	FTraitEventList InputEventList;
-	FTraitEventList OutputEventList;
+	FUpdateGraphContext UpdateGraphContext(HostInstance, ExecuteContext.GetDeltaTime());
+	UpdateGraphContext.SetBindingObject(RefPose.SkeletalMeshComponent);
+
+	FTraitEventList& InputEventList = UpdateGraphContext.GetInputEventList();
 
 	// A module can contain multiple graphs, we copy the input event list since it might be appended to during our update
 	{
@@ -141,9 +143,12 @@ FRigUnit_AnimNextRunAnimationGraph_v2_Execute()
 	// This reduces churn internally by avoiding a chunk to be repeatedly allocated and freed as we push/pop marks
 	MemStack.Alloc(size_t(FPageAllocator::SmallPageSize) + 1, 16);
 
-	UpdateGraph(HostInstance, ExecuteContext.GetDeltaTime(), InputEventList, OutputEventList);
-	EvaluateGraph(HostInstance, RefPose, LOD, Result);
-	
+	UpdateGraph(UpdateGraphContext);
+
+	FEvaluateGraphContext EvaluateGraphContext(HostInstance, RefPose, LOD);
+	EvaluateGraphContext.SetBindingObject(RefPose.SkeletalMeshComponent);
+	EvaluateGraph(EvaluateGraphContext, Result);
+
 	TRACE_ANIMNEXT_MODULE(ModuleInstance);
 	TRACE_ANIMNEXT_GRAPHINSTANCES(HostInstance);
 
@@ -164,6 +169,6 @@ FRigUnit_AnimNextRunAnimationGraph_v2_Execute()
 		}
 
 		// Append our output events
-		ModuleInstance.OutputEventList.Append(OutputEventList);
+		ModuleInstance.OutputEventList.Append(UpdateGraphContext.GetOutputEventList());
 	}
 }
diff --git a/Engine/Plugins/Experimental/AnimNextAnimGraph/Source/AnimNextAnimGraph/Private/PlayAnim/PlayAnimSlotTrait.cpp b/Engine/Plugins/Experimental/AnimNextAnimGraph/Source/AnimNextAnimGraph/Private/PlayAnim/PlayAnimSlotTrait.cpp
--- a/Engine/Plugins/Experimental/AnimNextAnimGraph/Source/AnimNextAnimGraph/Private/PlayAnim/PlayAnimSlotTrait.cpp
+++ b/Engine/Plugins/Experimental/AnimNextAnimGraph/Source/AnimNextAnimGraph/Private/PlayAnim/PlayAnimSlotTrait.cpp
@@ -398,7 +398,8 @@ namespace UE::AnimNext
 			if (SlotRequest.State == EPlayAnimRequestState::Active)
 			{
 				const FName EntryPoint = SlotRequest.AnimationGraph->DefaultEntryPoint;
-				SlotRequest.GraphInstance = SlotRequest.AnimationGraph->AllocateInstance(nullptr, &Binding.GetTraitPtr().GetNodeInstance()->GetOwner(), EntryPoint);
+				FAnimNextGraphInstance& Owner = Binding.GetTraitPtr().GetNodeInstance()->GetOwner();
+				SlotRequest.GraphInstance = SlotRequest.AnimationGraph->AllocateInstance(Owner.GetModuleInstance(), &Context, &Binding.GetTraitPtr().GetNodeInstance()->GetOwner(), EntryPoint);
 				SlotRequest.GraphInstance->BindPublicVariables(SlotRequest.Payload.Get());
 
 				// TODO: Validate that our child implements the ITimeline interface
diff --git a/Engine/Plugins/Experimental/AnimNextAnimGraph/Source/AnimNextAnimGraph/Private/TraitInterfaces/IEvaluate.cpp b/Engine/Plugins/Experimental/AnimNextAnimGraph/Source/AnimNextAnimGraph/Private/TraitInterfaces/IEvaluate.cpp
--- a/Engine/Plugins/Experimental/AnimNextAnimGraph/Source/AnimNextAnimGraph/Private/TraitInterfaces/IEvaluate.cpp
+++ b/Engine/Plugins/Experimental/AnimNextAnimGraph/Source/AnimNextAnimGraph/Private/TraitInterfaces/IEvaluate.cpp
@@ -91,12 +91,12 @@ namespace UE::AnimNext
 		}
 	};
 
-	FEvaluationProgram EvaluateGraph(const FAnimNextGraphInstance& GraphInstance)
+	FEvaluationProgram EvaluateGraph(const FEvaluateGraphContext& EvaluateGraphContext)
 	{
-		return EvaluateGraph(GraphInstance.GetGraphRootPtr());
+		return EvaluateGraph(EvaluateGraphContext, EvaluateGraphContext.GetGraphInstance().GetGraphRootPtr());
 	}
 
-	FEvaluationProgram EvaluateGraph(const FWeakTraitPtr& GraphRootPtr)
+	FEvaluationProgram EvaluateGraph(const FEvaluateGraphContext& EvaluateGraphContext, const FWeakTraitPtr& GraphRootPtr)
 	{
 		SCOPED_NAMED_EVENT(AnimNext_EvaluateGraph, FColor::Orange);
 		
@@ -113,6 +113,7 @@ namespace UE::AnimNext
 		FChildrenArray Children;
 
 		FEvaluateTraversalContext TraversalContext(EvaluationProgram);
+		TraversalContext.SetBindingObject(EvaluateGraphContext.GetBindingObject());
 		TraversalContext.BindTo(GraphRootPtr);
 
 		// Add the graph root to kick start the evaluation process
@@ -221,9 +222,13 @@ namespace UE::AnimNext
 		return EvaluationProgram;
 	}
 
-	void EvaluateGraph(FAnimNextGraphInstance& GraphInstance, const FReferencePose& RefPose, int32 GraphLODLevel, FAnimNextGraphLODPose& OutputPose)
+	void EvaluateGraph(const FEvaluateGraphContext& EvaluateGraphContext, FAnimNextGraphLODPose& OutputPose)
 	{
-		const FEvaluationProgram EvaluationProgram = EvaluateGraph(GraphInstance);
+		FAnimNextGraphInstance& GraphInstance = EvaluateGraphContext.GetGraphInstance();
+		const FReferencePose& RefPose = EvaluateGraphContext.GetRefPose();
+		const int32 GraphLODLevel = EvaluateGraphContext.GetGraphLODLevel();
+
+		const FEvaluationProgram EvaluationProgram = EvaluateGraph(EvaluateGraphContext);
 
 		TRACE_ANIMNEXT_EVALUATIONPROGRAM(EvaluationProgram, GraphInstance);
 
diff --git a/Engine/Plugins/Experimental/AnimNextAnimGraph/Source/AnimNextAnimGraph/Private/TraitInterfaces/IUpdate.cpp b/Engine/Plugins/Experimental/AnimNextAnimGraph/Source/AnimNextAnimGraph/Private/TraitInterfaces/IUpdate.cpp
--- a/Engine/Plugins/Experimental/AnimNextAnimGraph/Source/AnimNextAnimGraph/Private/TraitInterfaces/IUpdate.cpp
+++ b/Engine/Plugins/Experimental/AnimNextAnimGraph/Source/AnimNextAnimGraph/Private/TraitInterfaces/IUpdate.cpp
@@ -472,10 +472,12 @@ namespace UE::AnimNext
 	// perform as much useful work as possible while waiting for memory, hiding its slow latency by fully
 	// leveraging out-of-order CPU execution.
 
-	void UpdateGraph(FAnimNextGraphInstance& GraphInstance, float DeltaTime, FTraitEventList& InputEventList, FTraitEventList& OutputEventList)
+	void UpdateGraph(FUpdateGraphContext& UpdateGraphContext)
 	{
 		SCOPED_NAMED_EVENT(AnimNext_UpdateGraph, FColor::Orange);
 		
+		FAnimNextGraphInstance& GraphInstance = UpdateGraphContext.GetGraphInstance(); 
+
 		if (!GraphInstance.IsValid())
 		{
 			return;	// Nothing to update
@@ -486,9 +488,14 @@ namespace UE::AnimNext
 			return;	// We can only update starting at the root
 		}
 
+		FTraitEventList& InputEventList = UpdateGraphContext.GetInputEventList();
+		FTraitEventList& OutputEventList = UpdateGraphContext.GetOutputEventList();
+		const float DeltaTime = UpdateGraphContext.GetDeltaTime();
+
 		FUpdateTraversalContext TraversalContext;
 		TraversalContext.InputEventList = &InputEventList;
 		TraversalContext.OutputEventList = &OutputEventList;
+		TraversalContext.SetBindingObject(UpdateGraphContext.GetBindingObject());
 
 		FMemStack& MemStack = TraversalContext.GetMemStack();
 		FMemMark Mark(MemStack);
diff --git a/Engine/Plugins/Experimental/AnimNextAnimGraph/Source/AnimNextAnimGraph/Private/Traits/BlendSpacePlayer.cpp b/Engine/Plugins/Experimental/AnimNextAnimGraph/Source/AnimNextAnimGraph/Private/Traits/BlendSpacePlayer.cpp
--- a/Engine/Plugins/Experimental/AnimNextAnimGraph/Source/AnimNextAnimGraph/Private/Traits/BlendSpacePlayer.cpp
+++ b/Engine/Plugins/Experimental/AnimNextAnimGraph/Source/AnimNextAnimGraph/Private/Traits/BlendSpacePlayer.cpp
@@ -142,7 +142,8 @@ namespace UE::AnimNext
 						FSampleGraphState& TargetGraphState = InstanceData->SampleGraphs[CurSampleIndex];
 
 						const FName EntryPoint = NAME_None;
-						TargetGraphState.Instance = AnimationGraph->AllocateInstance(nullptr, &Binding.GetTraitPtr().GetNodeInstance()->GetOwner(), EntryPoint);
+						FAnimNextGraphInstance& Owner = Binding.GetTraitPtr().GetNodeInstance()->GetOwner();
+						TargetGraphState.Instance = AnimationGraph->AllocateInstance(Owner.GetModuleInstance(), &Context, &Binding.GetTraitPtr().GetNodeInstance()->GetOwner(), EntryPoint);
 						TargetGraphState.ChildPtr = TargetGraphState.Instance->GetGraphRootPtr();
 						TargetGraphState.GraphPayload = MoveTemp(GraphPayload);
 
diff --git a/Engine/Plugins/Experimental/AnimNextAnimGraph/Source/AnimNextAnimGraph/Private/Traits/BlendStackTrait.cpp b/Engine/Plugins/Experimental/AnimNextAnimGraph/Source/AnimNextAnimGraph/Private/Traits/BlendStackTrait.cpp
--- a/Engine/Plugins/Experimental/AnimNextAnimGraph/Source/AnimNextAnimGraph/Private/Traits/BlendStackTrait.cpp
+++ b/Engine/Plugins/Experimental/AnimNextAnimGraph/Source/AnimNextAnimGraph/Private/Traits/BlendStackTrait.cpp
@@ -219,7 +219,7 @@ namespace UE::AnimNext
 					//@TODO: Remove or implement entry points once we decide if we still need them.
 					const FName EntryPoint = NAME_None;
 					FAnimNextGraphInstance& Owner = Binding.GetTraitPtr().GetNodeInstance()->GetOwner();
-					Graph.Instance = Graph.Request.AnimationGraph->AllocateInstance(Owner.GetModuleInstance(), &Owner, EntryPoint);
+					Graph.Instance = Graph.Request.AnimationGraph->AllocateInstance(Owner.GetModuleInstance(), &Context, &Owner, EntryPoint);
 					FModuleInjectionDataInterfaceAdapter ModuleAdapter(Context.GetRootGraphInstance().GetModuleInstance(), Graph.Request.BindingModuleHandle);
 					IDataInterfaceHost* DataInterfaceHost = &ModuleAdapter;
 					Graph.Instance->BindPublicVariables(Graph.Request.GraphPayload.Get(), TConstArrayView<IDataInterfaceHost*>(&DataInterfaceHost, 1));
diff --git a/Engine/Plugins/Experimental/AnimNextAnimGraph/Source/AnimNextAnimGraph/Private/Traits/SubGraphHost.cpp b/Engine/Plugins/Experimental/AnimNextAnimGraph/Source/AnimNextAnimGraph/Private/Traits/SubGraphHost.cpp
--- a/Engine/Plugins/Experimental/AnimNextAnimGraph/Source/AnimNextAnimGraph/Private/Traits/SubGraphHost.cpp
+++ b/Engine/Plugins/Experimental/AnimNextAnimGraph/Source/AnimNextAnimGraph/Private/Traits/SubGraphHost.cpp
@@ -221,7 +221,8 @@ namespace UE::AnimNext
 
 			if (SubGraphEntry.State == ESlotState::ActiveWithGraph)
 			{
-				SubGraphEntry.GraphInstance = SubGraphEntry.AnimationGraph->AllocateInstance(nullptr, &Binding.GetTraitPtr().GetNodeInstance()->GetOwner(), SubGraphEntry.EntryPoint);
+				FAnimNextGraphInstance& Owner = Binding.GetTraitPtr().GetNodeInstance()->GetOwner();
+				SubGraphEntry.GraphInstance = SubGraphEntry.AnimationGraph->AllocateInstance(Owner.GetModuleInstance(), &Context, &Binding.GetTraitPtr().GetNodeInstance()->GetOwner(), SubGraphEntry.EntryPoint);
 			}
 		}
 	}
diff --git a/Engine/Plugins/Experimental/AnimNextAnimGraph/Source/AnimNextAnimGraph/Public/Graph/AnimNextAnimationGraph.h b/Engine/Plugins/Experimental/AnimNextAnimGraph/Source/AnimNextAnimGraph/Public/Graph/AnimNextAnimationGraph.h
--- a/Engine/Plugins/Experimental/AnimNextAnimGraph/Source/AnimNextAnimGraph/Public/Graph/AnimNextAnimationGraph.h
+++ b/Engine/Plugins/Experimental/AnimNextAnimGraph/Source/AnimNextAnimGraph/Public/Graph/AnimNextAnimationGraph.h
@@ -76,7 +76,7 @@ public:
 	// @param	InParentGraphInstance		The parent graph instance to use
 	// @param	InEntryPoint				The entry point to use. If this is NAME_None then the default entry point for this graph is used
 	// @return the allocated instance. This can be invalid if the entry point was not found or the graph's root trait is invalid.
-	TSharedPtr<FAnimNextGraphInstance> AllocateInstance(FAnimNextModuleInstance* InModuleInstance = nullptr, FAnimNextGraphInstance* InParentGraphInstance = nullptr, FName InEntryPoint = NAME_None) const;
+	TSharedPtr<FAnimNextGraphInstance> AllocateInstance(FAnimNextModuleInstance* InModuleInstance = nullptr, UE::AnimNext::FExecutionContext* InParentContext = nullptr, FAnimNextGraphInstance* InParentGraphInstance = nullptr, FName InEntryPoint = NAME_None) const;
 
 private:
 	// Loads the graph data from the provided archive buffer and returns true on success, false otherwise
diff --git a/Engine/Plugins/Experimental/AnimNextAnimGraph/Source/AnimNextAnimGraph/Public/TraitCore/ExecutionContext.h b/Engine/Plugins/Experimental/AnimNextAnimGraph/Source/AnimNextAnimGraph/Public/TraitCore/ExecutionContext.h
--- a/Engine/Plugins/Experimental/AnimNextAnimGraph/Source/AnimNextAnimGraph/Public/TraitCore/ExecutionContext.h
+++ b/Engine/Plugins/Experimental/AnimNextAnimGraph/Source/AnimNextAnimGraph/Public/TraitCore/ExecutionContext.h
@@ -200,6 +200,10 @@ namespace UE::AnimNext
 		// Returns the bound root graph instance
 		FAnimNextGraphInstance& GetRootGraphInstance() const { return *RootGraphInstance; }
 
+		// Returns the skeletal mesh component associated to the graph execution
+		void SetBindingObject(const TWeakObjectPtr<const USkeletalMeshComponent>& InBindingObject) { BindingObject = InBindingObject; }
+		const TWeakObjectPtr<const USkeletalMeshComponent> GetBindingObject() const { return BindingObject; }
+
 		// Returns the local thread memstack
 		FMemStack& GetMemStack() const { return MemStack; }
 
@@ -225,6 +229,9 @@ namespace UE::AnimNext
 		const FTrait* GetTrait(const FTraitTemplate& TraitDesc) const;
 
 	protected:
+		// The skeletal mesh component associated to the graph execution
+		TWeakObjectPtr<const USkeletalMeshComponent> BindingObject = nullptr;
+
 		// The memstack of the local thread
 		FMemStack& MemStack;
 
@@ -244,6 +251,8 @@ namespace UE::AnimNext
 		// be warmer in the CPU cache
 		Private::FScopedInterfaceEntry* FreeScopedInterfaceEntryStackHead = nullptr;
 
+
+
 		friend struct FTraitStackBinding;
 	};
 
diff --git a/Engine/Plugins/Experimental/AnimNextAnimGraph/Source/AnimNextAnimGraph/Public/TraitInterfaces/IEvaluate.h b/Engine/Plugins/Experimental/AnimNextAnimGraph/Source/AnimNextAnimGraph/Public/TraitInterfaces/IEvaluate.h
--- a/Engine/Plugins/Experimental/AnimNextAnimGraph/Source/AnimNextAnimGraph/Public/TraitInterfaces/IEvaluate.h
+++ b/Engine/Plugins/Experimental/AnimNextAnimGraph/Source/AnimNextAnimGraph/Public/TraitInterfaces/IEvaluate.h
@@ -10,9 +10,47 @@
 #include "EvaluationVM/KeyframeState.h"
 
 struct FAnimNextGraphLODPose;
+struct FRigUnit_AnimNextRunAnimationGraph_v1;
+struct FRigUnit_AnimNextRunAnimationGraph_v2;
+struct FAnimNode_AnimNextGraph;
 
 namespace UE::AnimNext
 {
+	/**
+	* FEvaluateGraphContext
+	*
+	* Data required to perform an EvaluateGraph
+	*
+	*/
+	struct FEvaluateGraphContext
+	{
+		FEvaluateGraphContext() = delete;
+		FEvaluateGraphContext(FAnimNextGraphInstance& InGraphInstance, const FReferencePose& InRefPose, int32 InGraphLODLevel)
+			: GraphInstance(InGraphInstance)
+			, RefPose(InRefPose)
+			, GraphLODLevel(InGraphLODLevel)
+		{
+		}
+
+		FAnimNextGraphInstance& GetGraphInstance() const { return GraphInstance; }
+		const FReferencePose& GetRefPose() const { return RefPose; }
+		int32 GetGraphLODLevel() const { return GraphLODLevel; }
+
+	protected:
+		const TWeakObjectPtr<const USkeletalMeshComponent>& GetBindingObject() const { return BindingObject; }
+		void SetBindingObject(const TWeakObjectPtr<const USkeletalMeshComponent>& InBindingObject) { BindingObject = InBindingObject; }
+
+		FAnimNextGraphInstance& GraphInstance;
+		const FReferencePose& RefPose;
+		int32 GraphLODLevel = 0;
+		TWeakObjectPtr<const USkeletalMeshComponent> BindingObject = nullptr;
+
+		friend ANIMNEXTANIMGRAPH_API FEvaluationProgram EvaluateGraph(const FEvaluateGraphContext& EvaluateGraphContext, const FWeakTraitPtr& GraphRootPtr);
+		friend struct ::FRigUnit_AnimNextRunAnimationGraph_v1;
+		friend struct ::FRigUnit_AnimNextRunAnimationGraph_v2;
+		friend struct ::FAnimNode_AnimNextGraph;
+	};
+
 	/**
 	 * FEvaluateTraversalContext
 	 *
@@ -35,7 +73,7 @@ namespace UE::AnimNext
 
 		FEvaluationProgram& EvaluationProgram;
 
-		friend ANIMNEXTANIMGRAPH_API FEvaluationProgram EvaluateGraph(const FWeakTraitPtr& GraphRootPtr);
+		friend ANIMNEXTANIMGRAPH_API FEvaluationProgram EvaluateGraph(const FEvaluateGraphContext& EvaluateGraphContext, const FWeakTraitPtr& GraphRootPtr);
 	};
 
 	/**
@@ -102,7 +140,7 @@ namespace UE::AnimNext
 	 *
 	 * @see IEvaluate::PreEvaluate, IEvaluate::PostEvaluate, IHierarchy::GetChildren
 	 */
-	[[nodiscard]] ANIMNEXTANIMGRAPH_API FEvaluationProgram EvaluateGraph(const FAnimNextGraphInstance& GraphInstance);
+	[[nodiscard]] ANIMNEXTANIMGRAPH_API FEvaluationProgram EvaluateGraph(const FEvaluateGraphContext& EvaluateGraphContext);
 
 	/**
 	 * Evaluates a sub-graph starting at its root and produces an evaluation program.
@@ -117,10 +155,10 @@ namespace UE::AnimNext
 	 *
 	 * @see IEvaluate::PreEvaluate, IEvaluate::PostEvaluate, IHierarchy::GetChildren
 	 */
-	[[nodiscard]] ANIMNEXTANIMGRAPH_API FEvaluationProgram EvaluateGraph(const FWeakTraitPtr& GraphRootPtr);
+	[[nodiscard]] ANIMNEXTANIMGRAPH_API FEvaluationProgram EvaluateGraph(const FEvaluateGraphContext& EvaluateGraphContext, const FWeakTraitPtr& GraphRootPtr);
 
 	/**
 	 * Evaluates a sub-graph starting at its root and produces an output pose
 	 */
-	ANIMNEXTANIMGRAPH_API void EvaluateGraph(FAnimNextGraphInstance& GraphInstance, const FReferencePose& RefPose, int32 GraphLODLevel, FAnimNextGraphLODPose& OutputPose);
+	ANIMNEXTANIMGRAPH_API void EvaluateGraph(const FEvaluateGraphContext& EvaluateGraphContext, FAnimNextGraphLODPose& OutputPose);
 }
diff --git a/Engine/Plugins/Experimental/AnimNextAnimGraph/Source/AnimNextAnimGraph/Public/TraitInterfaces/IUpdate.h b/Engine/Plugins/Experimental/AnimNextAnimGraph/Source/AnimNextAnimGraph/Public/TraitInterfaces/IUpdate.h
--- a/Engine/Plugins/Experimental/AnimNextAnimGraph/Source/AnimNextAnimGraph/Public/TraitInterfaces/IUpdate.h
+++ b/Engine/Plugins/Experimental/AnimNextAnimGraph/Source/AnimNextAnimGraph/Public/TraitInterfaces/IUpdate.h
@@ -11,6 +11,9 @@
 #include <type_traits>
 
 class FMemStack;
+struct FRigUnit_AnimNextRunAnimationGraph_v1;
+struct FRigUnit_AnimNextRunAnimationGraph_v2;
+struct FAnimNode_AnimNextGraph;
 
 namespace UE::AnimNext
 {
@@ -24,6 +27,50 @@ namespace UE::AnimNext
 		enum class FUpdateEventBookkeepingAction : uint8;
 	}
 
+	/**
+	* FUpdateGraphContext
+	*
+	* Data required to perform an UpdateGraph
+	*
+	*/
+	struct FUpdateGraphContext
+	{
+		FUpdateGraphContext() = delete;
+		FUpdateGraphContext(FAnimNextGraphInstance& InGraphInstance, float InDeltaTime)
+			: GraphInstance(InGraphInstance)
+			, DeltaTime(InDeltaTime)
+		{
+		}
+
+		void PushInputEvent(FAnimNextTraitEventPtr Event) { InputEventList.Push(MoveTemp(Event)); }
+
+		const FAnimNextGraphInstance& GetGraphInstance() const { return GraphInstance; }
+		FAnimNextGraphInstance& GetGraphInstance() { return GraphInstance; }
+
+		const FTraitEventList& GetInputEventList() const { return InputEventList; }
+		FTraitEventList& GetInputEventList() { return InputEventList; }
+
+		const FTraitEventList& GetOutputEventList() const { return OutputEventList; }
+		FTraitEventList& GetOutputEventList() { return OutputEventList; }
+
+		float GetDeltaTime() const { return DeltaTime; }
+
+	protected:
+		const TWeakObjectPtr<const USkeletalMeshComponent>& GetBindingObject() const { return BindingObject; }
+		void SetBindingObject(const TWeakObjectPtr<const USkeletalMeshComponent>& InBindingObject) { BindingObject = InBindingObject; }
+
+		FAnimNextGraphInstance& GraphInstance;
+		FTraitEventList InputEventList;
+		FTraitEventList OutputEventList;
+		float DeltaTime = 0.f;
+		TWeakObjectPtr<const USkeletalMeshComponent> BindingObject = nullptr;
+
+		friend ANIMNEXTANIMGRAPH_API void UpdateGraph(FUpdateGraphContext& UpdateGraphContext);
+		friend struct ::FRigUnit_AnimNextRunAnimationGraph_v1;
+		friend struct ::FRigUnit_AnimNextRunAnimationGraph_v2;
+		friend struct ::FAnimNode_AnimNextGraph;
+	};
+
 	/**
 	  * FTraitUpdateState
 	  * 
@@ -184,7 +231,7 @@ namespace UE::AnimNext
 		// The root node doesn't have a parent but we need a bookkeeping list regardless
 		Private::FUpdateEventBookkeepingList* RootParentBookkeepingEntryList = nullptr;
 
-		friend ANIMNEXTANIMGRAPH_API void UpdateGraph(FAnimNextGraphInstance& GraphInstance, float DeltaTime, FTraitEventList& InputEventList, FTraitEventList& OutputEventList);
+		friend ANIMNEXTANIMGRAPH_API void UpdateGraph(FUpdateGraphContext& UpdateGraphContext);
 		friend FUpdateTraversalQueue;
 	};
 
@@ -212,7 +259,7 @@ namespace UE::AnimNext
 		// emptied and pushed onto the update stack
 		Private::FUpdateEntry* QueuedUpdateStackHead = nullptr;
 
-		friend ANIMNEXTANIMGRAPH_API void UpdateGraph(FAnimNextGraphInstance& GraphInstance, float DeltaTime, FTraitEventList& InputEventList, FTraitEventList& OutputEventList);
+		friend ANIMNEXTANIMGRAPH_API void UpdateGraph(FUpdateGraphContext& UpdateGraphContext);
 		friend FUpdateTraversalContext;
 	};
 
@@ -326,5 +373,5 @@ namespace UE::AnimNext
 	 *
 	 * @see IUpdate::PreUpdate, IUpdate::PostUpdate, IHierarchy::GetChildren
 	 */
-	ANIMNEXTANIMGRAPH_API void UpdateGraph(FAnimNextGraphInstance& GraphInstance, float DeltaTime, FTraitEventList& InputEventList, FTraitEventList& OutputEventList);
+	ANIMNEXTANIMGRAPH_API void UpdateGraph(FUpdateGraphContext& UpdateGraphContext);
 }
diff --git a/Engine/Plugins/Experimental/AnimNextAnimGraph/Source/AnimNextAnimGraphTestSuite/Private/AnimNextAnimGraphTraitEventTest.cpp b/Engine/Plugins/Experimental/AnimNextAnimGraph/Source/AnimNextAnimGraphTestSuite/Private/AnimNextAnimGraphTraitEventTest.cpp
--- a/Engine/Plugins/Experimental/AnimNextAnimGraph/Source/AnimNextAnimGraphTestSuite/Private/AnimNextAnimGraphTraitEventTest.cpp
+++ b/Engine/Plugins/Experimental/AnimNextAnimGraph/Source/AnimNextAnimGraphTestSuite/Private/AnimNextAnimGraphTraitEventTest.cpp
@@ -308,12 +308,11 @@ bool FAnimationAnimNextRuntimeTest_GraphTraitEvent::RunTest(const FString& InPar
 				auto EventA = MakeTraitEvent<FTraitAnimGraphTest_EventA>();
 				auto EventB = MakeTraitEvent<FTraitAnimGraphTest_EventB>();
 
-				FTraitEventList InputEventList;
-				InputEventList.Push(EventA);
-				InputEventList.Push(EventB);
+				FUpdateGraphContext UpdateGraphContext(*GraphInstance.Get(), 0.0333f);
+				UpdateGraphContext.PushInputEvent(EventA);
+				UpdateGraphContext.PushInputEvent(EventB);
 
-				FTraitEventList OutputEventList;
-				UpdateGraph(*GraphInstance.Get(), 0.0333f, InputEventList, OutputEventList);
+				UpdateGraph(UpdateGraphContext);
 
 				UE_RETURN_ON_ERROR(EventA->VisitedTraits.Num() == 4, "FAnimationAnimNextRuntimeTest_GraphTraitEvent -> Unexpected number of traits visited");
 				AddErrorIfFalse(EventA->VisitedTraits[0] == NodeTemplateTraitListB[0], "FAnimationAnimNextRuntimeTest_GraphTraitEvent -> Unexpected trait visited");
@@ -333,11 +332,10 @@ bool FAnimationAnimNextRuntimeTest_GraphTraitEvent::RunTest(const FString& InPar
 				auto EventA = MakeTraitEvent<FTraitAnimGraphTest_EventA>();
 				EventA->bTestFlag = true;
 
-				FTraitEventList InputEventList;
-				InputEventList.Push(EventA);
+				FUpdateGraphContext UpdateGraphContext((*(GraphInstance.Get())), 0.0333f);
+				UpdateGraphContext.PushInputEvent(EventA);
 
-				FTraitEventList OutputEventList;
-				UpdateGraph(*GraphInstance.Get(), 0.0333f, InputEventList, OutputEventList);
+				UpdateGraph(UpdateGraphContext);
 
 				UE_RETURN_ON_ERROR(EventA->VisitedTraits.Num() == 2, "FAnimationAnimNextRuntimeTest_GraphTraitEvent -> Unexpected number of traits visited");
 				AddErrorIfFalse(EventA->VisitedTraits[0] == NodeTemplateTraitListB[0], "FAnimationAnimNextRuntimeTest_GraphTraitEvent -> Unexpected trait visited");
@@ -350,12 +348,11 @@ bool FAnimationAnimNextRuntimeTest_GraphTraitEvent::RunTest(const FString& InPar
 				EventA->bTestFlag = true;
 				auto EventB = MakeTraitEvent<FTraitAnimGraphTest_EventB>();
 
-				FTraitEventList InputEventList;
-				InputEventList.Push(EventA);
-				InputEventList.Push(EventB);
+				FUpdateGraphContext UpdateGraphContext(*GraphInstance.Get(), 0.0333f);
+				UpdateGraphContext.PushInputEvent(EventA);
+				UpdateGraphContext.PushInputEvent(EventB);
 
-				FTraitEventList OutputEventList;
-				UpdateGraph(*GraphInstance.Get(), 0.0333f, InputEventList, OutputEventList);
+				UpdateGraph(UpdateGraphContext);
 
 				UE_RETURN_ON_ERROR(EventA->VisitedTraits.Num() == 2, "FAnimationAnimNextRuntimeTest_GraphTraitEvent -> Unexpected number of traits visited");
 				AddErrorIfFalse(EventA->VisitedTraits[0] == NodeTemplateTraitListB[0], "FAnimationAnimNextRuntimeTest_GraphTraitEvent -> Unexpected trait visited");
@@ -373,11 +370,10 @@ bool FAnimationAnimNextRuntimeTest_GraphTraitEvent::RunTest(const FString& InPar
 				auto EventB = MakeTraitEvent<FTraitAnimGraphTest_EventB>();
 				EventB->bTestFlag0 = true;
 
-				FTraitEventList InputEventList;
-				InputEventList.Push(EventB);
+				FUpdateGraphContext UpdateGraphContext(*GraphInstance.Get(), 0.0333f);
+				UpdateGraphContext.PushInputEvent(EventB);
 
-				FTraitEventList OutputEventList;
-				UpdateGraph(*GraphInstance.Get(), 0.0333f, InputEventList, OutputEventList);
+				UpdateGraph(UpdateGraphContext);
 
 				UE_RETURN_ON_ERROR(EventB->VisitedTraits.Num() == 4, "FAnimationAnimNextRuntimeTest_GraphTraitEvent -> Unexpected number of traits visited");
 				AddErrorIfFalse(EventB->VisitedTraits[0] == NodeTemplateTraitListB[0], "FAnimationAnimNextRuntimeTest_GraphTraitEvent -> Unexpected trait visited");
@@ -398,11 +394,10 @@ bool FAnimationAnimNextRuntimeTest_GraphTraitEvent::RunTest(const FString& InPar
 				auto EventB = MakeTraitEvent<FTraitAnimGraphTest_EventB>();
 				EventB->bTestFlag1 = true;
 
-				FTraitEventList InputEventList;
-				InputEventList.Push(EventB);
+				FUpdateGraphContext UpdateGraphContext(*GraphInstance.Get(), 0.0333f);
+				UpdateGraphContext.PushInputEvent(EventB);
 
-				FTraitEventList OutputEventList;
-				UpdateGraph(*GraphInstance.Get(), 0.0333f, InputEventList, OutputEventList);
+				UpdateGraph(UpdateGraphContext);
 
 				UE_RETURN_ON_ERROR(EventB->VisitedTraits.Num() == 4, "FAnimationAnimNextRuntimeTest_GraphTraitEvent -> Unexpected number of traits visited");
 				AddErrorIfFalse(EventB->VisitedTraits[0] == NodeTemplateTraitListB[0], "FAnimationAnimNextRuntimeTest_GraphTraitEvent -> Unexpected trait visited");
diff --git a/Engine/Plugins/Experimental/AnimNextAnimGraph/Source/AnimNextAnimGraphTestSuite/Private/AnimNextAnimGraphTraitGraphTest.cpp b/Engine/Plugins/Experimental/AnimNextAnimGraph/Source/AnimNextAnimGraphTestSuite/Private/AnimNextAnimGraphTraitGraphTest.cpp
--- a/Engine/Plugins/Experimental/AnimNextAnimGraph/Source/AnimNextAnimGraphTestSuite/Private/AnimNextAnimGraphTraitGraphTest.cpp
+++ b/Engine/Plugins/Experimental/AnimNextAnimGraph/Source/AnimNextAnimGraphTestSuite/Private/AnimNextAnimGraphTraitGraphTest.cpp
@@ -541,10 +541,10 @@ bool FAnimationAnimNextRuntimeTest_GraphExecute::RunTest(const FString& InParame
 		GraphInstance->GetExtendedExecuteContext().SetRuntimeSettings(RuntimeSettings);
 
 		{
-			UE::AnimNext::FTraitEventList InputEventList;
-			UE::AnimNext::FTraitEventList OutputEventList;
-			UE::AnimNext::UpdateGraph(*GraphInstance.Get(), 1.0f / 30.0f, InputEventList, OutputEventList);
-			(void)UE::AnimNext::EvaluateGraph(*GraphInstance.Get());
+			FUpdateGraphContext UpdateGraphContext(*GraphInstance.Get(), 1.0f / 30.0f);
+			UE::AnimNext::UpdateGraph(UpdateGraphContext);
+			const UE::AnimNext::FEvaluateGraphContext EvaluateGraphContext(*GraphInstance.Get(), UE::AnimNext::FReferencePose(), 0);
+			(void)UE::AnimNext::EvaluateGraph(EvaluateGraphContext);
 		}
 
 		AddErrorIfFalse(Messages.Num() == 7, "FAnimationAnimNextRuntimeTest_GraphExecute - > Unexpected message count");
@@ -679,10 +679,10 @@ bool FAnimationAnimNextRuntimeTest_GraphExecuteLatent::RunTest(const FString& In
 		GraphInstance->GetExtendedExecuteContext().SetRuntimeSettings(RuntimeSettings);
 		
 		{
-			UE::AnimNext::FTraitEventList InputEventList;
-			UE::AnimNext::FTraitEventList OutputEventList;
-			UE::AnimNext::UpdateGraph(*GraphInstance.Get(), 1.0f / 30.0f, InputEventList, OutputEventList);
-			(void)UE::AnimNext::EvaluateGraph(*GraphInstance.Get());
+			FUpdateGraphContext UpdateGraphContext(*GraphInstance.Get(), 1.0f / 30.0f);
+			UE::AnimNext::UpdateGraph(UpdateGraphContext);
+			const UE::AnimNext::FEvaluateGraphContext EvaluateGraphContext(*GraphInstance.Get(), UE::AnimNext::FReferencePose(), 0);
+			(void)UE::AnimNext::EvaluateGraph(EvaluateGraphContext);
 		}
 
 		AddErrorIfFalse(Messages.Num() == 7, "FAnimationAnimNextRuntimeTest_GraphExecuteLatent - > Unexpected message count");
diff --git a/Engine/Plugins/Experimental/AnimNextAnimGraph/Source/AnimNextAnimGraphTestSuite/Private/AnimNextAnimGraphTraitInterfacesTest.cpp b/Engine/Plugins/Experimental/AnimNextAnimGraph/Source/AnimNextAnimGraphTestSuite/Private/AnimNextAnimGraphTraitInterfacesTest.cpp
--- a/Engine/Plugins/Experimental/AnimNextAnimGraph/Source/AnimNextAnimGraphTestSuite/Private/AnimNextAnimGraphTraitInterfacesTest.cpp
+++ b/Engine/Plugins/Experimental/AnimNextAnimGraph/Source/AnimNextAnimGraphTestSuite/Private/AnimNextAnimGraphTraitInterfacesTest.cpp
@@ -804,9 +804,8 @@ bool FAnimationAnimNextRuntimeTest_IUpdate::RunTest(const FString& InParameters)
 			Private::UpdatedTraits = &UpdatedTraits;
 
 			// Call pre/post update on our graph
-			FTraitEventList InputEventList;
-			FTraitEventList OutputEventList;
-			UpdateGraph(*GraphInstance.Get(), 0.0333f, InputEventList, OutputEventList);
+			FUpdateGraphContext UpdateGraphContext(*GraphInstance.Get(), 0.0333f);
+			UpdateGraph(UpdateGraphContext);
 
 			AddErrorIfFalse(UpdatedTraits.Num() == 6, "FAnimationAnimNextRuntimeTest_IUpdate -> Expected 6 nodes to have been visited during the update traversal");
 			AddErrorIfFalse(UpdatedTraits[0] == FTraitWithChildren::TraitUID, "FAnimationAnimNextRuntimeTest_IUpdate -> Unexpected update order");		// NodeC
@@ -946,7 +945,8 @@ bool FAnimationAnimNextRuntimeTest_IEvaluate::RunTest(const FString& InParameter
 			Private::EvaluatedTraits = &EvaluatedTraits;
 
 			// Call pre/post evaluate on our graph
-			(void)EvaluateGraph(*GraphInstance.Get());
+			UE::AnimNext::FEvaluateGraphContext EvaluateGraphContext(*GraphInstance.Get(), UE::AnimNext::FReferencePose(), 0);
+			(void)EvaluateGraph(EvaluateGraphContext);
 
 			AddErrorIfFalse(EvaluatedTraits.Num() == 6, "FAnimationAnimNextRuntimeTest_IEvaluate -> Expected 6 nodes to have been visited during the evaluate traversal");
 			AddErrorIfFalse(EvaluatedTraits[0] == FTraitWithChildren::TraitUID, "FAnimationAnimNextRuntimeTest_IEvaluate -> Unexpected evaluate order");		// NodeC
@@ -1087,9 +1087,8 @@ bool FAnimationAnimNextRuntimeTest_IScopedInterface::RunTest(const FString& InPa
 
 			// Call pre/post update on our graph
 			Private::AutoPopTag = true;
-			FTraitEventList InputEventList;
-			FTraitEventList OutputEventList;
-			UpdateGraph(*GraphInstance.Get(), 0.0333f, InputEventList, OutputEventList);
+			FUpdateGraphContext UpdateGraphContext(*GraphInstance.Get(), 0.0333f);
+			UpdateGraph(UpdateGraphContext);
 
 			AddErrorIfFalse(IsTagInScope.Num() == 8, "FAnimationAnimNextRuntimeTest_IScopedInterface -> Unexpected number of entries");
 			AddErrorIfFalse(IsTagInScope[0] == false, "FAnimationAnimNextRuntimeTest_IScopedInterface -> Unexpected scoped tag state");		// NodeC::PreUpdate (template 0)
@@ -1105,7 +1104,7 @@ bool FAnimationAnimNextRuntimeTest_IScopedInterface::RunTest(const FString& InPa
 			Private::AutoPopTag = false;
 
 			// Call pre/post update on our graph
-			UpdateGraph(*GraphInstance.Get(), 0.0333f, InputEventList, OutputEventList);
+			UpdateGraph(UpdateGraphContext);
 
 			AddErrorIfFalse(IsTagInScope.Num() == 8, "FAnimationAnimNextRuntimeTest_IScopedInterface -> Unexpected number of entries");
 			AddErrorIfFalse(IsTagInScope[0] == false, "FAnimationAnimNextRuntimeTest_IScopedInterface -> Unexpected scoped tag state");		// NodeC::PreUpdate (template 0)
diff --git a/Engine/Plugins/Experimental/AnimNextControlRig/Source/AnimNextControlRig/Internal/ControlRigTrait.h b/Engine/Plugins/Experimental/AnimNextControlRig/Source/AnimNextControlRig/Internal/ControlRigTrait.h
--- a/Engine/Plugins/Experimental/AnimNextControlRig/Source/AnimNextControlRig/Internal/ControlRigTrait.h
+++ b/Engine/Plugins/Experimental/AnimNextControlRig/Source/AnimNextControlRig/Internal/ControlRigTrait.h
@@ -214,6 +214,7 @@ struct FControlRigTrait : FBaseTrait, IEvaluate, IUpdate, IHierarchy, IGarbageCo
 		void Destruct(const FExecutionContext& Context, const FTraitBinding& Binding);
 
 		static UObject* GetAnimContext(const FExecutionContext& Context);
+		static const USkeletalMeshComponent* GetBindableObject(const FExecutionContext& Context);
 		static UClass* GetTargetClass(const FSharedData* SharedData);
 
 		// Returns the property, variable name and the memory of the variable (as target) and the memory of the latent property (as source)
@@ -252,7 +253,7 @@ struct FControlRigTrait : FBaseTrait, IEvaluate, IUpdate, IHierarchy, IGarbageCo
 protected:
 	static UControlRig* GetControlRig(FInstanceData* InstanceData);
 
-	static bool SetupControlRig(UObject* InAnimContext, UClass* InControlRigClass, FInstanceData* InstanceData);
+	static bool SetupControlRig(UObject* InAnimContext, UObject* InBindableObject, UClass* InControlRigClass, FInstanceData* InstanceData);
 
 	static UClass* GetTargetClass(const FSharedData* SharedData, const TTraitBinding<IUpdate>& Binding);
 
diff --git a/Engine/Plugins/Experimental/AnimNextControlRig/Source/AnimNextControlRig/Private/ControlRigTrait.cpp b/Engine/Plugins/Experimental/AnimNextControlRig/Source/AnimNextControlRig/Private/ControlRigTrait.cpp
--- a/Engine/Plugins/Experimental/AnimNextControlRig/Source/AnimNextControlRig/Private/ControlRigTrait.cpp
+++ b/Engine/Plugins/Experimental/AnimNextControlRig/Source/AnimNextControlRig/Private/ControlRigTrait.cpp
@@ -139,21 +139,24 @@ void FControlRigTrait::FInstanceData::Construct(const FExecutionContext& Context
 
 	const FSharedData* SharedData = Binding.GetSharedData<FSharedData>();
 
-	UObject* AnimContext = GetAnimContext(Context);
 	if (UClass* ControlRigClass = GetTargetClass(SharedData))
 	{
-		// we also need access to the properties when running construction event
-		//ControlRigVariableMappings.PropagateCustomInputProperties(ControlRig); // there is no control rig at this point, how is this supposed to work ?
-
-		if (ControlRigClass && ensure(SetupControlRig(AnimContext, ControlRigClass, this)))
+		if (UObject* AnimContext = GetAnimContext(Context))
 		{
-			OnInitializedHandle = ControlRig->OnInitialized_AnyThread().AddRaw(this, &FControlRigTrait::FInstanceData::HandleOnInitialized_AnyThread);
-		}
+			// we also need access to the properties when running construction event
+			//ControlRigVariableMappings.PropagateCustomInputProperties(ControlRig); // there is no control rig at this point, how is this supposed to work ?
+
+			USkeletalMeshComponent* BindableObject = const_cast<USkeletalMeshComponent*>(FInstanceData::GetBindableObject(Context));
+			if (ControlRigClass && ensure(SetupControlRig(AnimContext, BindableObject, ControlRigClass, this)))
+			{
+				OnInitializedHandle = ControlRig->OnInitialized_AnyThread().AddRaw(this, &FControlRigTrait::FInstanceData::HandleOnInitialized_AnyThread);
+			}
 
-		ControlRigVariableMappings.ResetCurvesInputToControlCache();
-		ControlRigVariableMappings.CacheCurveMappings(SharedData->InputMapping, SharedData->OutputMapping, ControlRig->GetHierarchy());
+			ControlRigVariableMappings.ResetCurvesInputToControlCache();
+			ControlRigVariableMappings.CacheCurveMappings(SharedData->InputMapping, SharedData->OutputMapping, ControlRig->GetHierarchy());
 
-		InitializeCustomProperties(Binding, SharedData);
+			InitializeCustomProperties(Binding, SharedData);
+		}
 	}
 
 #if UE_ENABLE_DEBUG_DRAWING
@@ -194,6 +197,13 @@ UObject* FControlRigTrait::FInstanceData::GetAnimContext(const FExecutionContext
 	return AnimContext;
 }
 
+const USkeletalMeshComponent* FControlRigTrait::FInstanceData::GetBindableObject(const FExecutionContext& Context)
+{
+	const USkeletalMeshComponent* BindableObject = Context.GetBindingObject().Get();
+	
+	return BindableObject;
+}
+
 UClass* FControlRigTrait::FInstanceData::GetTargetClass(const FSharedData* SharedData)
 {
 	if (SharedData != nullptr)
@@ -506,17 +516,13 @@ void FControlRigTrait::PostEvaluate(FEvaluateTraversalContext& Context, const TT
 
 	if (UControlRig* ControlRig = GetControlRig(InstanceData))
 	{
-		// Skeletal mesh required to be able to create a FBoneContainer
-		if (UObject* AnimContext = FInstanceData::GetAnimContext(Context))
-		{
-			InstanceData->SkeletalMeshComponent = Cast<USkeletalMeshComponent>(FControlRigObjectBinding::GetBindableObject(AnimContext));
-			if (InstanceData->SkeletalMeshComponent != nullptr)
-			{
+		// The transform is used by the task to "transform" the debug drawings inside Control Rig
 #if UE_ENABLE_DEBUG_DRAWING
-				InstanceData->ComponentTransform = InstanceData->SkeletalMeshComponent->GetComponentTransform();
-#endif
-			}
+		if (const USkeletalMeshComponent* BindableObject = FInstanceData::GetBindableObject(Context))
+		{
+			InstanceData->ComponentTransform = BindableObject->GetComponentTransform();
 		}
+#endif
 
 		Context.AppendTask(FAnimNextControlRigTask::Make(SharedData, InstanceData));
 	}
@@ -560,7 +566,7 @@ UControlRig* FControlRigTrait::GetControlRig(FInstanceData* InstanceData)
 	return nullptr;
 }
 
-bool FControlRigTrait::SetupControlRig(UObject* InAnimContext, UClass* InControlRigClass, FInstanceData* InstanceData)
+bool FControlRigTrait::SetupControlRig(UObject* InAnimContext, UObject* InBindableObject, UClass* InControlRigClass, FInstanceData* InstanceData)
 {
 	if (InstanceData->ControlRig == nullptr)
 	{
@@ -575,17 +581,17 @@ bool FControlRigTrait::SetupControlRig(UObject* InAnimContext, UClass* InControl
 
 			UControlRig* ControlRig = InstanceData->ControlRig.Get();
 			ControlRig->SetObjectBinding(MakeShared<FControlRigObjectBinding>());
-			if (InAnimContext != nullptr)
-			{
-				ControlRig->GetObjectBinding()->BindToObject(InAnimContext);
-			}
+
+			UObject* BindableObject = InBindableObject != nullptr ? FControlRigObjectBinding::GetBindableObject(InBindableObject) : FControlRigObjectBinding::GetBindableObject(InAnimContext);
+			check (BindableObject != nullptr);
+
+			ControlRig->GetObjectBinding()->BindToObject(BindableObject);
 			ControlRig->Initialize(true);
 			ControlRig->RequestInit();
 			
 			InstanceData->ControlRigHierarchyMappings.ResetRefPoseSetterHash();
 			
 			// register bindable object as data source (used for To World / From World transformations)
-			UObject* BindableObject = FControlRigObjectBinding::GetBindableObject(InAnimContext);
 			ControlRig->GetDataSourceRegistry()->RegisterDataSource(UControlRig::OwnerComponent, BindableObject);
 
 			//UpdateInputOutputMappingIfRequired(ControlRig, InRequiredBones);
