Fixed crash in  Control Rig Trait when a UUSerSruct is edited, as the post edit reinits the VM and fores a new initialziation

diff --git a/Engine/Plugins/Experimental/AnimNextControlRig/Source/AnimNextControlRig/Private/ControlRigTrait.cpp b/Engine/Plugins/Experimental/AnimNextControlRig/Source/AnimNextControlRig/Private/ControlRigTrait.cpp
--- a/Engine/Plugins/Experimental/AnimNextControlRig/Source/AnimNextControlRig/Private/ControlRigTrait.cpp
+++ b/Engine/Plugins/Experimental/AnimNextControlRig/Source/AnimNextControlRig/Private/ControlRigTrait.cpp
@@ -321,6 +321,9 @@ int32 FControlRigTrait::FInstanceData::GetExposedControlsData(const UE::AnimNext
 void FControlRigTrait::FInstanceData::HandleOnInitialized_AnyThread(URigVMHost*, const FName&)
 {
 	ControlRigHierarchyMappings.ResetRefPoseSetterHash();
+	#if WITH_EDITOR
+		bRegenerateVariableMappings = true;	// required as FRigVMEditorModule::PreChange (UUSerStructs) recreates VM memory and requests a re-init, which recreates controls
+	#endif
 }
 
 void FControlRigTrait::FInstanceData::InitializeCustomProperties(const FTraitBinding& Binding, const FSharedData* SharedData)
