Fix for UAF Control Rig getting desync when a rig is compiled while PIE is running, by performing a full reinit on object reinstance

#bs-asset-validation-passed: 1648016 v2_AAAAAGgFLaI=
#tests Windows full cook + run

diff --git a/Engine/Plugins/Experimental/AnimNextControlRig/Source/AnimNextControlRig/Internal/ControlRigTrait.h b/Engine/Plugins/Experimental/AnimNextControlRig/Source/AnimNextControlRig/Internal/ControlRigTrait.h
--- a/Engine/Plugins/Experimental/AnimNextControlRig/Source/AnimNextControlRig/Internal/ControlRigTrait.h
+++ b/Engine/Plugins/Experimental/AnimNextControlRig/Source/AnimNextControlRig/Internal/ControlRigTrait.h
@@ -213,6 +213,9 @@ struct FControlRigTrait : FBaseTrait, IEvaluate, IUpdate, IHierarchy, IGarbageCo
 		void Construct(const FExecutionContext& Context, const FTraitBinding& Binding);
 		void Destruct(const FExecutionContext& Context, const FTraitBinding& Binding);
 
+		void InitializeControlRig(const FExecutionContext& Context, const FTraitBinding& Binding);
+		void DestroyControlRig(const FExecutionContext& Context, const FTraitBinding& Binding);
+
 		static UObject* GetAnimContext(const FExecutionContext& Context);
 		static const USkeletalMeshComponent* GetBindableObject(const FExecutionContext& Context);
 		static UClass* GetTargetClass(const FSharedData* SharedData);
@@ -227,6 +230,8 @@ struct FControlRigTrait : FBaseTrait, IEvaluate, IUpdate, IHierarchy, IGarbageCo
 
 #if WITH_EDITOR
 		void OnObjectsReinstanced(const TMap<UObject*, UObject*>& OldToNewInstanceMap);
+		bool bRefreshBindableObject = false;
+		bool bReinitializeControlRig = false;
 		bool bRegenerateVariableMappings = false;
 #endif
 	};
@@ -253,7 +258,8 @@ struct FControlRigTrait : FBaseTrait, IEvaluate, IUpdate, IHierarchy, IGarbageCo
 protected:
 	static UControlRig* GetControlRig(FInstanceData* InstanceData);
 
-	static bool SetupControlRig(UObject* InAnimContext, UObject* InBindableObject, UClass* InControlRigClass, FInstanceData* InstanceData);
+	static bool CreateControlRig(UObject* InAnimContext, UObject* InBindableObject, UClass* InControlRigClass, FInstanceData* InstanceData);
+	static bool SetBindableObject(UControlRig* ControlRig, UObject* InAnimContext, UObject* InBindableObject);
 
 	static UClass* GetTargetClass(const FSharedData* SharedData, const TTraitBinding<IUpdate>& Binding);
 
diff --git a/Engine/Plugins/Experimental/AnimNextControlRig/Source/AnimNextControlRig/Private/ControlRigTrait.cpp b/Engine/Plugins/Experimental/AnimNextControlRig/Source/AnimNextControlRig/Private/ControlRigTrait.cpp
--- a/Engine/Plugins/Experimental/AnimNextControlRig/Source/AnimNextControlRig/Private/ControlRigTrait.cpp
+++ b/Engine/Plugins/Experimental/AnimNextControlRig/Source/AnimNextControlRig/Private/ControlRigTrait.cpp
@@ -129,6 +129,8 @@ void FControlRigTrait::FInstanceData::Construct(const FExecutionContext& Context
 {
 	FTrait::FInstanceData::Construct(Context, Binding);
 
+	const FSharedData* SharedData = Binding.GetSharedData<FSharedData>();
+
 	IGarbageCollection::RegisterWithGC(Context, Binding);
 
 #if WITH_EDITOR
@@ -137,25 +139,15 @@ void FControlRigTrait::FInstanceData::Construct(const FExecutionContext& Context
 	}
 #endif
 
-	const FSharedData* SharedData = Binding.GetSharedData<FSharedData>();
-
 	if (UClass* ControlRigClass = GetTargetClass(SharedData))
 	{
 		if (UObject* AnimContext = GetAnimContext(Context))
 		{
-			// we also need access to the properties when running construction event
-			//ControlRigVariableMappings.PropagateCustomInputProperties(ControlRig); // there is no control rig at this point, how is this supposed to work ?
-
 			USkeletalMeshComponent* BindableObject = const_cast<USkeletalMeshComponent*>(FInstanceData::GetBindableObject(Context));
-			if (ControlRigClass && ensure(SetupControlRig(AnimContext, BindableObject, ControlRigClass, this)))
+			if (ensure(CreateControlRig(GetAnimContext(Context), BindableObject, SharedData->ControlRigClass, this)))
 			{
-				OnInitializedHandle = ControlRig->OnInitialized_AnyThread().AddRaw(this, &FControlRigTrait::FInstanceData::HandleOnInitialized_AnyThread);
+				InitializeControlRig(Context, Binding);
 			}
-
-			ControlRigVariableMappings.ResetCurvesInputToControlCache();
-			ControlRigVariableMappings.CacheCurveMappings(SharedData->InputMapping, SharedData->OutputMapping, ControlRig->GetHierarchy());
-
-			InitializeCustomProperties(Binding, SharedData);
 		}
 	}
 
@@ -174,6 +166,41 @@ void FControlRigTrait::FInstanceData::Destruct(const FExecutionContext& Context,
 	UE::AnimNext::ControlRig::FAnimNextControlRigModule::OnObjectsReinstanced.Remove(OnObjectsReinstancedHandle);
 #endif
 
+	DestroyControlRig(Context, Binding);
+}
+
+void FControlRigTrait::FInstanceData::InitializeControlRig(const FExecutionContext& Context, const FTraitBinding& Binding)
+{
+	const FSharedData* SharedData = Binding.GetSharedData<FSharedData>();
+
+	if (ControlRig)
+	{
+		// Provide available properties to the construction event
+		InitializeCustomProperties(Binding, SharedData);
+		ControlRigVariableMappings.PropagateCustomInputProperties(ControlRig);
+
+		ControlRig->Initialize(true);
+		ControlRig->RequestInit();
+
+		ControlRigHierarchyMappings.ResetRefPoseSetterHash();
+
+		ControlRigVariableMappings.ResetCurvesInputToControlCache();
+		ControlRigHierarchyMappings.ResetRefPoseSetterHash();
+		ControlRigVariableMappings.CacheCurveMappings(SharedData->InputMapping, SharedData->OutputMapping, ControlRig->GetHierarchy());
+
+		// Re-init Custom Properties after construction, as new controls could be created and might have to be remapped
+		InitializeCustomProperties(Binding, SharedData);
+		ControlRigVariableMappings.PropagateCustomInputProperties(ControlRig);
+
+		//UpdateInputOutputMappingIfRequired(ControlRig, InRequiredBones);
+
+		ControlRig->OnInitialized_AnyThread().RemoveAll(this);
+		OnInitializedHandle = ControlRig->OnInitialized_AnyThread().AddRaw(this, &FControlRigTrait::FInstanceData::HandleOnInitialized_AnyThread);
+	}
+}
+
+void FControlRigTrait::FInstanceData::DestroyControlRig(const FExecutionContext& Context, const FTraitBinding& Binding)
+{
 	if (ControlRig.Get() != nullptr)
 	{
 		if (OnInitializedHandle.IsValid())
@@ -181,6 +208,7 @@ void FControlRigTrait::FInstanceData::Destruct(const FExecutionContext& Context,
 			ControlRig->OnInitialized_AnyThread().Remove(OnInitializedHandle);
 			OnInitializedHandle.Reset();
 		}
+		ControlRig->MarkAsGarbage();
 		ControlRig = nullptr;
 	}
 }
@@ -345,17 +373,17 @@ void FControlRigTrait::FInstanceData::OnObjectsReinstanced(const TMap<UObject*,
 		for (const TPair<UObject*, UObject*>& Pair : OldToNewInstanceMap)
 		{
 			const UObject* NewObject = Pair.Value;
-			if ((NewObject == nullptr) || !NewObject->IsA<UControlRig>())
+			if((NewObject == nullptr) || (NewObject->GetOuter() != ControlRig->GetOuter()) || (!NewObject->IsA<UControlRig>()))
 			{
 				continue;
 			}
-			// NewObject->GetFName() != ControlRig->GetFName()
-			ControlRigVariableMappings.ResetCustomProperties(); // avoid crashes when properties go away
-			bRegenerateVariableMappings = true;
 
-			ControlRig->OnInitialized_AnyThread().RemoveAll(this);
-			OnInitializedHandle = ControlRig->OnInitialized_AnyThread().AddRaw(this, &FControlRigTrait::FInstanceData::HandleOnInitialized_AnyThread);
-			break;
+			if (NewObject->GetClass() == ControlRig->GetClass())
+			{
+				bRefreshBindableObject = true;
+				bReinitializeControlRig = true;
+				break;
+			}
 		}
 	}
 }
@@ -497,6 +525,20 @@ void FControlRigTrait::PreUpdate(FUpdateTraversalContext& Context, const TTraitB
 		InstanceData->Input = Context.AllocateNodeInstance(Binding, SharedData->Input);
 	}
 
+#if WITH_EDITOR
+	if (InstanceData->bRefreshBindableObject) // AnimNext full recompile Thaw function does not set a binding object
+	{
+		UObject* AnimContext = InstanceData->GetAnimContext(Context);
+		USkeletalMeshComponent* BindableObject = const_cast<USkeletalMeshComponent*>(FInstanceData::GetBindableObject(Context));
+		InstanceData->bRefreshBindableObject = !SetBindableObject(InstanceData->ControlRig, AnimContext, BindableObject);
+	}
+	if (InstanceData->bReinitializeControlRig)
+	{
+		InstanceData->InitializeControlRig(Context, Binding);
+		InstanceData->bReinitializeControlRig = false;
+	}
+#endif
+
 	if (UControlRig* ControlRig = GetControlRig(InstanceData))
 	{
 #if WITH_EDITOR
@@ -572,39 +614,44 @@ UControlRig* FControlRigTrait::GetControlRig(FInstanceData* InstanceData)
 	return nullptr;
 }
 
-bool FControlRigTrait::SetupControlRig(UObject* InAnimContext, UObject* InBindableObject, UClass* InControlRigClass, FInstanceData* InstanceData)
+bool FControlRigTrait::CreateControlRig(UObject* InAnimContext, UObject* InBindableObject, UClass* InControlRigClass, FInstanceData* InstanceData)
 {
-	if (InstanceData->ControlRig == nullptr)
+	if (InstanceData->ControlRig == nullptr && InControlRigClass != nullptr)
 	{
-		if (InControlRigClass != nullptr)
+		// Let's make sure the GC isn't running when we try to create a new Control Rig.
 		{
-			// Let's make sure the GC isn't running when we try to create a new Control Rig.
-			{
-				FGCScopeGuard GCGuard;
-				InstanceData->ControlRig = NewObject<UControlRig>(InAnimContext, InControlRigClass);
-				InstanceData->ControlRig->AtomicallyClearInternalFlags(EInternalObjectFlags::Async);
-			}
+			FGCScopeGuard GCGuard;
+			InstanceData->ControlRig = NewObject<UControlRig>(InAnimContext, InControlRigClass);
+			InstanceData->ControlRig->AtomicallyClearInternalFlags(EInternalObjectFlags::Async);
+		}
 
-			UControlRig* ControlRig = InstanceData->ControlRig.Get();
-			ControlRig->SetObjectBinding(MakeShared<FControlRigObjectBinding>());
+		const bool bSuccess = SetBindableObject(InstanceData->ControlRig, InAnimContext, InBindableObject);
+#if WITH_EDITOR
+		InstanceData->bRefreshBindableObject = !bSuccess;
+#endif
+	}
 
-			UObject* BindableObject = InBindableObject != nullptr ? FControlRigObjectBinding::GetBindableObject(InBindableObject) : FControlRigObjectBinding::GetBindableObject(InAnimContext);
-			check (BindableObject != nullptr);
+	return InstanceData->ControlRig != nullptr;
+}
 
-			ControlRig->GetObjectBinding()->BindToObject(BindableObject);
-			ControlRig->Initialize(true);
-			ControlRig->RequestInit();
-			
-			InstanceData->ControlRigHierarchyMappings.ResetRefPoseSetterHash();
-			
-			// register bindable object as data source (used for To World / From World transformations)
-			ControlRig->GetDataSourceRegistry()->RegisterDataSource(UControlRig::OwnerComponent, BindableObject);
 
-			//UpdateInputOutputMappingIfRequired(ControlRig, InRequiredBones);
-		}
-	}
+bool FControlRigTrait::SetBindableObject(UControlRig* ControlRig, UObject* InAnimContext, UObject* InBindableObject)
+{
+	if (ControlRig != nullptr && ensure(InAnimContext) && InBindableObject)
+	{
+		ControlRig->SetObjectBinding(MakeShared<FControlRigObjectBinding>());
+
+		UObject* ObjectToBind = InBindableObject != nullptr ? InBindableObject : FControlRigObjectBinding::GetBindableObject(InAnimContext);
+		check(ObjectToBind != nullptr);
+
+		ControlRig->GetObjectBinding()->BindToObject(ObjectToBind);
 
-	return InstanceData->ControlRig.Get() != nullptr;
+		// register bindable object as data source (used for To World / From World transformations)
+		ControlRig->GetDataSourceRegistry()->RegisterDataSource(UControlRig::OwnerComponent, ObjectToBind);
+
+		return true;
+	}
+	return false;
 }
 
 #if WITH_EDITOR
