Fix for Control Rig not being able to expose variables due to the trait name being different from expected.
Added code to search for the different names the Control Rig Trait is added into the stack (the context menu uses the display name while the trait editor uses the struct name)

diff --git a/Engine/Plugins/Experimental/AnimNextControlRig/Source/AnimNextControlRigEditor/Internal/ControlRigTraitCustomization.h b/Engine/Plugins/Experimental/AnimNextControlRig/Source/AnimNextControlRigEditor/Internal/ControlRigTraitCustomization.h
--- a/Engine/Plugins/Experimental/AnimNextControlRig/Source/AnimNextControlRigEditor/Internal/ControlRigTraitCustomization.h
+++ b/Engine/Plugins/Experimental/AnimNextControlRig/Source/AnimNextControlRigEditor/Internal/ControlRigTraitCustomization.h
@@ -10,6 +10,7 @@
 class URigVMPin;
 class IDetailCategoryBuilder;
 class UEdGraphNode;
+class URigVMNode;
 struct FRigVMVariableMappingInfo;
 struct FControlRigTraitSharedData;
 
@@ -45,6 +46,8 @@ private:
 
 	void OnObjectsReinstanced(const TMap<UObject*, UObject*>& OldToNewInstanceMap);
 
+	FString FindControlRigTraitPinName(URigVMNode* ModelNode);
+
 	TSharedPtr<IPropertyHandle> StructPropertyHandle;
 	TSharedPtr<FStructOnScope> ScopedControlRigTraitSharedData;
 
diff --git a/Engine/Plugins/Experimental/AnimNextControlRig/Source/AnimNextControlRigEditor/Private/ControlRigTraitCustomization.cpp b/Engine/Plugins/Experimental/AnimNextControlRig/Source/AnimNextControlRigEditor/Private/ControlRigTraitCustomization.cpp
--- a/Engine/Plugins/Experimental/AnimNextControlRig/Source/AnimNextControlRigEditor/Private/ControlRigTraitCustomization.cpp
+++ b/Engine/Plugins/Experimental/AnimNextControlRig/Source/AnimNextControlRigEditor/Private/ControlRigTraitCustomization.cpp
@@ -20,6 +20,8 @@
 #include "AnimNextControlRigModule.h"
 #include "ScopedTransaction.h"
 #include "K2Node.h"
+#include "TraitCore/TraitRegistry.h"
+#include "RigVMStringUtils.h"
 
 #define LOCTEXT_NAMESPACE "ControlRigTraitSharedDataCustomization"
 
@@ -122,9 +124,12 @@ void FControlRigTraitSharedDataCustomization::OnPropertyExposeCheckboxChanged(EC
 						static FName ExposedPropertyVariablesName = GET_MEMBER_NAME_CHECKED(FControlRigTraitSharedData, ExposedPropertyVariableNames);
 						static FName ExposedPropertyControlsName = GET_MEMBER_NAME_CHECKED(FControlRigTraitSharedData, ExposedPropertyControlNames);
 						static FName ExposedPropertyControlTypesName = GET_MEMBER_NAME_CHECKED(FControlRigTraitSharedData, ExposedPropertyControlTypes);
-						static FString ExposedVariablesSubPath(TEXT("FControlRigTrait.") + ExposedPropertyVariablesName.ToString());
-						static FString ExpòsedControlsSubPath(TEXT("FControlRigTrait.") + ExposedPropertyControlsName.ToString());
-						static FString ExpòsedControlTypesSubPath(TEXT("FControlRigTrait.") + ExposedPropertyControlTypesName.ToString());
+						
+						const FString ControlRigTraitName = FindControlRigTraitPinName(ModelNode);
+						const FString ControlRigTraitPrefix(ControlRigTraitName + TEXT("."));
+						const FString ExposedVariablesSubPath(ControlRigTraitPrefix + ExposedPropertyVariablesName.ToString());
+						FString ExpòsedControlsSubPath(ControlRigTraitPrefix + ExposedPropertyControlsName.ToString());
+						FString ExpòsedControlTypesSubPath(ControlRigTraitPrefix + ExposedPropertyControlTypesName.ToString());
 
 						static const FText ExposePropertyToPinText = LOCTEXT("ExposePropertyToPin", "Expose Property to Pin");
 						static const FText RemovePropertyPinText = LOCTEXT("RemovePropertyOptionalPin", "Removed Property Optional Pin");
@@ -386,6 +391,44 @@ void FControlRigTraitSharedDataCustomization::OnObjectsReinstanced(const TMap<UO
 	}
 }
 
+FString FControlRigTraitSharedDataCustomization::FindControlRigTraitPinName(URigVMNode* ModelNode)
+{
+	const FTraitRegistry& TraitRegistry = FTraitRegistry::Get();
+
+	FString ControlRigTraitName;
+	FString ControlRigTraitDisplayName;
+	if (const FTrait* Trait = TraitRegistry.Find(FControlRigTrait::TraitUID)) // the context menu in the traits create the traits using a sanitized display name
+	{
+		ControlRigTraitName = Trait->GetTraitName();
+		FString DisplayNameMetadata;
+		Trait->GetTraitSharedDataStruct()->GetStringMetaDataHierarchical(FRigVMStruct::DisplayNameMetaName, &DisplayNameMetadata);
+		ControlRigTraitDisplayName = DisplayNameMetadata.IsEmpty() ? Trait->GetTraitName() : DisplayNameMetadata;
+		RigVMStringUtils::SanitizeName(ControlRigTraitDisplayName, false, false, URigVMSchema::GetMaxNameLength());
+	}
+	else
+	{
+		ControlRigTraitName = FString(TEXT("FControlRigTrait"));
+	}
+
+	FString Tmp;
+	for (const auto& Pin : ModelNode->GetPins())
+	{
+		Pin->GetName(Tmp);
+
+		if (Tmp.Contains(ControlRigTraitName))
+		{
+			return Tmp;
+		}
+
+		if (Tmp.Contains(ControlRigTraitDisplayName)) 
+		{
+			return Tmp;
+		}
+	}
+
+	return FString();
+}
+
 } // end namespace UE::AnimNext::Editor
 
 
