Fix for Control Rig Trait losing exposed controls due to load dependencies in editor

diff --git a/Engine/Plugins/Animation/ControlRig/Source/ControlRig/Private/Rigs/RigHierarchy.cpp b/Engine/Plugins/Animation/ControlRig/Source/ControlRig/Private/Rigs/RigHierarchy.cpp
--- a/Engine/Plugins/Animation/ControlRig/Source/ControlRig/Private/Rigs/RigHierarchy.cpp
+++ b/Engine/Plugins/Animation/ControlRig/Source/ControlRig/Private/Rigs/RigHierarchy.cpp
@@ -3056,6 +3056,11 @@ int32 URigHierarchy::GetPoseVersion(const FRigTransformElement* InTransformEleme
 FEdGraphPinType URigHierarchy::GetControlPinType(FRigControlElement* InControlElement) const
 {
 	check(InControlElement);
+	return GetControlPinType(InControlElement->Settings.ControlType);
+}
+
+FEdGraphPinType URigHierarchy::GetControlPinType(ERigControlType ControlType)
+{
 	LLM_SCOPE_BYNAME(TEXT("Animation/ControlRig"));
 
 	// local copy of UEdGraphSchema_K2::PC_ ... static members
@@ -3067,7 +3072,7 @@ FEdGraphPinType URigHierarchy::GetControlPinType(FRigControlElement* InControlEl
 
 	FEdGraphPinType PinType;
 
-	switch(InControlElement->Settings.ControlType)
+	switch(ControlType)
 	{
 		case ERigControlType::Bool:
 		{
diff --git a/Engine/Plugins/Animation/ControlRig/Source/ControlRig/Public/Rigs/RigHierarchy.h b/Engine/Plugins/Animation/ControlRig/Source/ControlRig/Public/Rigs/RigHierarchy.h
--- a/Engine/Plugins/Animation/ControlRig/Source/ControlRig/Public/Rigs/RigHierarchy.h
+++ b/Engine/Plugins/Animation/ControlRig/Source/ControlRig/Public/Rigs/RigHierarchy.h
@@ -2968,6 +2968,13 @@ public:
 	 */
 	CONTROLRIG_API FEdGraphPinType GetControlPinType(FRigControlElement* InControlElement) const;
 
+	/**
+	 * Returns the pin type to use for a control type
+	 * @param InControlElement The control to return the pin type for
+	 * @return The pin type
+	 */
+	CONTROLRIG_API static FEdGraphPinType GetControlPinType(ERigControlType ControlType);
+
 	/**
 	 * Returns the default value to use for a pin for a control
 	 * @param InControlElement The control to return the pin default value for
diff --git a/Engine/Plugins/Experimental/AnimNextControlRig/Source/AnimNextControlRig/Internal/ControlRigTrait.h b/Engine/Plugins/Experimental/AnimNextControlRig/Source/AnimNextControlRig/Internal/ControlRigTrait.h
--- a/Engine/Plugins/Experimental/AnimNextControlRig/Source/AnimNextControlRig/Internal/ControlRigTrait.h
+++ b/Engine/Plugins/Experimental/AnimNextControlRig/Source/AnimNextControlRig/Internal/ControlRigTrait.h
@@ -136,11 +136,6 @@ struct FControlRigTraitSharedData : public FAnimNextTraitSharedData
 	UPROPERTY(meta = (Hidden))
 	TMap<FName, FName> OutputMapping;
 
-	// This is the array of latent input properties
-	// This is computed at load time based on the selected ControlRig class
-	UPROPERTY(meta = (Hidden))
-	TArray<FName> LatentProperties;
-
 	UPROPERTY(meta = (Hidden))
 	TArray<FName> ExposedPropertyVariableNames;
 
@@ -150,6 +145,14 @@ struct FControlRigTraitSharedData : public FAnimNextTraitSharedData
 	UPROPERTY(meta = (Hidden))
 	TArray<ERigControlType> ExposedPropertyControlTypes;
 
+	UPROPERTY(meta = (Hidden))
+	TArray<FString> ExposedPropertyControlDefaultValues;
+
+	// This is the array of latent input properties
+	// This is computed at load time based on the selected ControlRig class
+	UPROPERTY(Transient, meta = (Hidden))
+	TArray<FName> LatentProperties;
+
 	// This is the array that maps a latent property to its memory layout in the trait instance data
 	// This is computed at load time based on the selected ControlRig class
 	UPROPERTY(Transient, meta = (Hidden))
diff --git a/Engine/Plugins/Experimental/AnimNextControlRig/Source/AnimNextControlRig/Private/ControlRigTrait.cpp b/Engine/Plugins/Experimental/AnimNextControlRig/Source/AnimNextControlRig/Private/ControlRigTrait.cpp
--- a/Engine/Plugins/Experimental/AnimNextControlRig/Source/AnimNextControlRig/Private/ControlRigTrait.cpp
+++ b/Engine/Plugins/Experimental/AnimNextControlRig/Source/AnimNextControlRig/Private/ControlRigTrait.cpp
@@ -698,29 +698,47 @@ void FControlRigTrait::GetProgrammaticPins(FAnimNextTraitSharedData* InSharedDat
 			}
 
 			// --- Exposed Controls ---
-			if (SharedData->ExposedPropertyControlNames.Num() > 0)
+			const int32 NumExposedControls = SharedData->ExposedPropertyControlNames.Num();
+			if (NumExposedControls > 0)
 			{
-				if (USkeleton* PreviewSkeleton = SharedData->GetPreviewSkeleton())
+				if (NumExposedControls == SharedData->ExposedPropertyControlDefaultValues.Num() && NumExposedControls == SharedData->ExposedPropertyControlTypes.Num())
 				{
-					// Obtain the controls from the RigControlsData helper. This will instantiate a rig using the provided class and cache the controls until the class changes
-					const TArray<FControlRigIOMapping::FControlsInfo>& RigControls = RigControlsData.GetControls(SharedData->ControlRigClass, PreviewSkeleton);
-
-					for (const FControlRigIOMapping::FControlsInfo& Control : RigControls)
+					for (int32 ControlIndex = 0; ControlIndex < NumExposedControls; ControlIndex++)
 					{
-						const FName& ControlName = Control.Name;
-						if (!SharedData->ExposedPropertyControlNames.Contains(ControlName))	// Only process exposed controls
-						{
-							continue;
-						}
+						const FName& ControlName = SharedData->ExposedPropertyControlNames[ControlIndex];
+						const FString& ControlDefaultValue = SharedData->ExposedPropertyControlDefaultValues[ControlIndex];
+						const TRigVMTypeIndex TypeIndex = RigVMTypeUtils::TypeIndexFromPinType(URigHierarchy::GetControlPinType(SharedData->ExposedPropertyControlTypes[ControlIndex]));
 
-						const FString& ControlDefaultValue = Control.DefaultValue;
-						const TRigVMTypeIndex TypeIndex = RigVMTypeUtils::TypeIndexFromPinType(Control.PinType);
 						if (ensure(TypeIndex != INDEX_NONE))
 						{
 							OutPinArray.AddPin(InController, InParentPinIndex, ControlName, ERigVMPinDirection::Input, TypeIndex, ControlDefaultValue, ERigVMPinDefaultValueType::AutoDetect, nullptr, nullptr, true);
 						}
 					}
 				}
+				else
+				{
+					if (USkeleton* PreviewSkeleton = SharedData->GetPreviewSkeleton())
+					{
+						// Obtain the controls from the RigControlsData helper. This will instantiate a rig using the provided class and cache the controls until the class changes
+						const TArray<FControlRigIOMapping::FControlsInfo>& RigControls = RigControlsData.GetControls(SharedData->ControlRigClass, PreviewSkeleton);
+
+						for (const FControlRigIOMapping::FControlsInfo& Control : RigControls)
+						{
+							const FName& ControlName = Control.Name;
+							if (!SharedData->ExposedPropertyControlNames.Contains(ControlName))	// Only process exposed controls
+							{
+								continue;
+							}
+
+							const FString& ControlDefaultValue = Control.DefaultValue;
+							const TRigVMTypeIndex TypeIndex = RigVMTypeUtils::TypeIndexFromPinType(Control.PinType);
+							if (ensure(TypeIndex != INDEX_NONE))
+							{
+								OutPinArray.AddPin(InController, InParentPinIndex, ControlName, ERigVMPinDirection::Input, TypeIndex, ControlDefaultValue, ERigVMPinDefaultValueType::AutoDetect, nullptr, nullptr, true);
+							}
+						}
+					}
+				}
 			}
 		}
 	}
diff --git a/Engine/Plugins/Experimental/AnimNextControlRig/Source/AnimNextControlRigEditor/Private/ControlRigTraitCustomization.cpp b/Engine/Plugins/Experimental/AnimNextControlRig/Source/AnimNextControlRigEditor/Private/ControlRigTraitCustomization.cpp
--- a/Engine/Plugins/Experimental/AnimNextControlRig/Source/AnimNextControlRigEditor/Private/ControlRigTraitCustomization.cpp
+++ b/Engine/Plugins/Experimental/AnimNextControlRig/Source/AnimNextControlRigEditor/Private/ControlRigTraitCustomization.cpp
@@ -124,12 +124,14 @@ void FControlRigTraitSharedDataCustomization::OnPropertyExposeCheckboxChanged(EC
 						static FName ExposedPropertyVariablesName = GET_MEMBER_NAME_CHECKED(FControlRigTraitSharedData, ExposedPropertyVariableNames);
 						static FName ExposedPropertyControlsName = GET_MEMBER_NAME_CHECKED(FControlRigTraitSharedData, ExposedPropertyControlNames);
 						static FName ExposedPropertyControlTypesName = GET_MEMBER_NAME_CHECKED(FControlRigTraitSharedData, ExposedPropertyControlTypes);
+						static FName ExposedPropertyControlDefaultValuesName = GET_MEMBER_NAME_CHECKED(FControlRigTraitSharedData, ExposedPropertyControlDefaultValues);
 						
 						const FString ControlRigTraitName = FindControlRigTraitPinName(ModelNode);
 						const FString ControlRigTraitPrefix(ControlRigTraitName + TEXT("."));
 						const FString ExposedVariablesSubPath(ControlRigTraitPrefix + ExposedPropertyVariablesName.ToString());
-						FString ExpòsedControlsSubPath(ControlRigTraitPrefix + ExposedPropertyControlsName.ToString());
-						FString ExpòsedControlTypesSubPath(ControlRigTraitPrefix + ExposedPropertyControlTypesName.ToString());
+						const FString ExposedControlsSubPath(ControlRigTraitPrefix + ExposedPropertyControlsName.ToString());
+						const FString ExposedControlTypesSubPath(ControlRigTraitPrefix + ExposedPropertyControlTypesName.ToString());
+						const FString ExposedControlDefaultValuesSubPath(ControlRigTraitPrefix + ExposedPropertyControlDefaultValuesName.ToString());
 
 						static const FText ExposePropertyToPinText = LOCTEXT("ExposePropertyToPin", "Expose Property to Pin");
 						static const FText RemovePropertyPinText = LOCTEXT("RemovePropertyOptionalPin", "Removed Property Optional Pin");
@@ -147,13 +149,15 @@ void FControlRigTraitSharedDataCustomization::OnPropertyExposeCheckboxChanged(EC
 						}
 
 						URigVMPin* VariablesPin = ModelNode->FindPin(ExposedVariablesSubPath);
-						URigVMPin* ControlsPin = ModelNode->FindPin(ExpòsedControlsSubPath);
-						URigVMPin* ControlsTypesPin = ModelNode->FindPin(ExpòsedControlTypesSubPath);
-						if (ensure(VariablesPin && ControlsPin && ControlsTypesPin))
+						URigVMPin* ControlsPin = ModelNode->FindPin(ExposedControlsSubPath);
+						URigVMPin* ControlsTypesPin = ModelNode->FindPin(ExposedControlTypesSubPath);
+						URigVMPin* ControlsDefaultValuesPin = ModelNode->FindPin(ExposedControlDefaultValuesSubPath);
+						if (ensure(VariablesPin && ControlsPin && ControlsTypesPin && ControlsDefaultValuesPin))
 						{
 							ControlRigTraitSharedData->ExposedPropertyVariableNames.Reset(CustomPinProperties.Num());
 							ControlRigTraitSharedData->ExposedPropertyControlNames.Reset(CustomPinProperties.Num());
 							ControlRigTraitSharedData->ExposedPropertyControlTypes.Reset(CustomPinProperties.Num());
+							ControlRigTraitSharedData->ExposedPropertyControlDefaultValues.Reset(CustomPinProperties.Num());
 
 							for (const FOptionalPinFromProperty& OptionalPin : CustomPinProperties)
 							{
@@ -170,6 +174,7 @@ void FControlRigTraitSharedDataCustomization::OnPropertyExposeCheckboxChanged(EC
 										{
 											ControlRigTraitSharedData->ExposedPropertyControlNames.Add(OptionalPin.PropertyName);
 											ControlRigTraitSharedData->ExposedPropertyControlTypes.Add(ControlInfo->ControlType);
+											ControlRigTraitSharedData->ExposedPropertyControlDefaultValues.Add(ControlInfo->DefaultValue);
 										}
 									}
 								}
@@ -194,6 +199,7 @@ void FControlRigTraitSharedDataCustomization::OnPropertyExposeCheckboxChanged(EC
 							SetPinArrayDefaultValue(Controller, VariablesPin, ControlRigTraitSharedData, ExposedPropertyVariablesName, &ControlRigTraitSharedData->ExposedPropertyVariableNames);
 							SetPinArrayDefaultValue(Controller, ControlsPin, ControlRigTraitSharedData, ExposedPropertyControlsName, &ControlRigTraitSharedData->ExposedPropertyControlNames);
 							SetPinArrayDefaultValue(Controller, ControlsTypesPin, ControlRigTraitSharedData, ExposedPropertyControlTypesName, &ControlRigTraitSharedData->ExposedPropertyControlTypes);
+							SetPinArrayDefaultValue(Controller, ControlsDefaultValuesPin, ControlRigTraitSharedData, ExposedPropertyControlDefaultValuesName, &ControlRigTraitSharedData->ExposedPropertyControlDefaultValues);
 						}
 
 						Controller->RepopulatePinsOnNode(ModelNode, true, false, true);
