Control Rig Trait : Cache custom property mappings during latent pin construction in order to avoid the cost of retrieving the list of publc variables from the CDO during destruction in the game thread

diff --git a/Engine/Plugins/Experimental/AnimNextControlRig/Source/AnimNextControlRig/Internal/ControlRigTrait.h b/Engine/Plugins/Experimental/AnimNextControlRig/Source/AnimNextControlRig/Internal/ControlRigTrait.h
--- a/Engine/Plugins/Experimental/AnimNextControlRig/Source/AnimNextControlRig/Internal/ControlRigTrait.h
+++ b/Engine/Plugins/Experimental/AnimNextControlRig/Source/AnimNextControlRig/Internal/ControlRigTrait.h
@@ -196,6 +196,7 @@ struct FControlRigTrait : FBaseTrait, IEvaluate, IUpdate, IHierarchy, IGarbageCo
 
 		ControlRig::FAnimNextControlRigHierarchyMappings ControlRigHierarchyMappings;
 		FControlRigVariableMappings ControlRigVariableMappings;
+		FControlRigVariableMappings::FCustomPropertyMappings PropertyMappings;
 
 		TWeakObjectPtr<UNodeMappingContainer> NodeMappingContainer;
 		FControlRigIOSettings InputSettings;
diff --git a/Engine/Plugins/Experimental/AnimNextControlRig/Source/AnimNextControlRig/Private/ControlRigTrait.cpp b/Engine/Plugins/Experimental/AnimNextControlRig/Source/AnimNextControlRig/Private/ControlRigTrait.cpp
--- a/Engine/Plugins/Experimental/AnimNextControlRig/Source/AnimNextControlRig/Private/ControlRigTrait.cpp
+++ b/Engine/Plugins/Experimental/AnimNextControlRig/Source/AnimNextControlRig/Private/ControlRigTrait.cpp
@@ -32,11 +32,12 @@ void FControlRigTraitSharedData::ConstructLatentProperties(const UE::AnimNext::F
 	const FControlRigTraitSharedData* SharedData = Binding.GetSharedData<FControlRigTraitSharedData>();
 	if (SharedData && SharedData->ControlRigClass.Get() != nullptr)
 	{
-		FControlRigVariableMappings::FCustomPropertyMappings PropertyMappings;
-		UE::AnimNext::FControlRigTrait::FInstanceData::GetExposedVariablesData(Binding, SharedData, PropertyMappings);
-		UE::AnimNext::FControlRigTrait::FInstanceData::GetExposedControlsData(Binding, SharedData, PropertyMappings);
+		UE::AnimNext::FControlRigTrait::FInstanceData* InstanceData = Binding.GetInstanceData<UE::AnimNext::FControlRigTrait::FInstanceData>();
 
-		const TArray<FControlRigVariableMappings::FCustomPropertyData>& Mappings = PropertyMappings.GetMappings();
+		UE::AnimNext::FControlRigTrait::FInstanceData::GetExposedVariablesData(Binding, SharedData, InstanceData->PropertyMappings);
+		UE::AnimNext::FControlRigTrait::FInstanceData::GetExposedControlsData(Binding, SharedData, InstanceData->PropertyMappings);
+
+		const TArray<FControlRigVariableMappings::FCustomPropertyData>& Mappings = InstanceData->PropertyMappings.GetMappings();
 		const int32 NumProperties = Mappings.Num();
 
 		for (const FControlRigVariableMappings::FCustomPropertyData& Mapping : Mappings)
@@ -58,11 +59,9 @@ void FControlRigTraitSharedData::DestructLatentProperties(const UE::AnimNext::FT
 	const FControlRigTraitSharedData* SharedData = Binding.GetSharedData<FControlRigTraitSharedData>();
 	if (SharedData && SharedData->ControlRigClass.Get() != nullptr)
 	{
-		FControlRigVariableMappings::FCustomPropertyMappings PropertyMappings;
-		UE::AnimNext::FControlRigTrait::FInstanceData::GetExposedVariablesData(Binding, SharedData, PropertyMappings);
-		UE::AnimNext::FControlRigTrait::FInstanceData::GetExposedControlsData(Binding, SharedData, PropertyMappings);
+		UE::AnimNext::FControlRigTrait::FInstanceData* InstanceData = Binding.GetInstanceData<UE::AnimNext::FControlRigTrait::FInstanceData>();
 
-		const TArray<FControlRigVariableMappings::FCustomPropertyData>& Mappings = PropertyMappings.GetMappings();
+		const TArray<FControlRigVariableMappings::FCustomPropertyData>& Mappings = InstanceData->PropertyMappings.GetMappings();
 		const int32 NumProperties = Mappings.Num();
 
 		for (const FControlRigVariableMappings::FCustomPropertyData& Mapping : Mappings)
@@ -356,13 +355,10 @@ void FControlRigTrait::FInstanceData::HandleOnInitialized_AnyThread(URigVMHost*,
 
 void FControlRigTrait::FInstanceData::InitializeCustomProperties(const FTraitBinding& Binding, const FSharedData* SharedData)
 {
-	// Obtain exposed variables and controls data, in order to init the variable mappings
-	FControlRigVariableMappings::FCustomPropertyMappings PropertyMappings;
-	GetExposedVariablesData(Binding, SharedData, PropertyMappings);
-	GetExposedControlsData(Binding, SharedData, PropertyMappings);
-
+	UE::AnimNext::FControlRigTrait::FInstanceData* InstanceData = Binding.GetInstanceData<UE::AnimNext::FControlRigTrait::FInstanceData>();
+	 
 	// Setup mappings using the latent pin memory as source (we have to copy from latent pin to external variable / rig control)
-	ControlRigVariableMappings.InitializeCustomProperties(ControlRig, PropertyMappings);
+	ControlRigVariableMappings.InitializeCustomProperties(ControlRig, InstanceData->PropertyMappings);
 }
 
 #if WITH_EDITOR
@@ -534,6 +530,10 @@ void FControlRigTrait::PreUpdate(FUpdateTraversalContext& Context, const TTraitB
 	}
 	if (InstanceData->bReinitializeControlRig)
 	{
+		InstanceData->PropertyMappings.Reset();
+		UE::AnimNext::FControlRigTrait::FInstanceData::GetExposedVariablesData(Binding, SharedData, InstanceData->PropertyMappings);
+		UE::AnimNext::FControlRigTrait::FInstanceData::GetExposedControlsData(Binding, SharedData, InstanceData->PropertyMappings);
+
 		InstanceData->InitializeControlRig(Context, Binding);
 		InstanceData->bReinitializeControlRig = false;
 	}
