Removed stack allocator memmark from ControlRigTask, as it creates a memory after use due to the VM surviving the data is created after the mark

diff --git a/Engine/Plugins/Experimental/AnimNextControlRig/Source/AnimNextControlRig/Private/ControlRigTask.cpp b/Engine/Plugins/Experimental/AnimNextControlRig/Source/AnimNextControlRig/Private/ControlRigTask.cpp
--- a/Engine/Plugins/Experimental/AnimNextControlRig/Source/AnimNextControlRig/Private/ControlRigTask.cpp
+++ b/Engine/Plugins/Experimental/AnimNextControlRig/Source/AnimNextControlRig/Private/ControlRigTask.cpp
@@ -59,8 +59,6 @@ void FAnimNextControlRigTask::ExecuteControlRig(UE::AnimNext::FEvaluationVM& VM,
 
 	if (ControlRig)
 	{
-		FMemMark Mark(FMemStack::Get());
-
 		// Before we start modifying the RigHierarchy, we need to lock the rig to avoid corrupted state
 		// TODO : Is this really needed ? 
 		//FScopeLock LockRig(&ControlRig->GetEvaluateMutex());
